ChaosFlesh - FFleshCollection
- Property system for the flesh asset.

diff --git a/Engine/Plugins/Experimental/ChaosFlesh/ChaosFlesh.uplugin b/Engine/Plugins/Experimental/ChaosFlesh/ChaosFlesh.uplugin
new file mode 100644
--- /dev/null
+++ b/Engine/Plugins/Experimental/ChaosFlesh/ChaosFlesh.uplugin
@@ -0,0 +1,31 @@
+{
+    "FileVersion" : 3,
+    "Version" : 1,
+    "VersionName" : "0.1",
+    "FriendlyName" : "Chaos Flesh",
+    "Description" : "Chaos Flesh Simulation",
+    "Category": "Physics",
+    "CreatedBy" : "Epic Games, Inc.",
+    "CreatedByURL" : "http://epicgames.com",
+    "DocsURL" : "",
+    "MarketplaceURL" : "",
+    "SupportURL" : "",
+    "EnabledByDefault" : false,
+    "CanContainContent" : true,
+    "IsBetaVersion" : true,
+	"IsExperimentalVersion": true,
+    "Installed" : false,
+    "Modules": [
+		{
+			"Name": "ChaosFlesh",
+			"Type": "Runtime",
+			"LoadingPhase": "PreDefault"
+		},
+    ],
+	"Plugins": [
+		{
+			"Name": "ProceduralMeshComponent",
+			"Enabled": true
+		},
+	]
+}
diff --git a/Engine/Plugins/Experimental/ChaosFlesh/Source/ChaosFlesh/ChaosFlesh.Build.cs b/Engine/Plugins/Experimental/ChaosFlesh/Source/ChaosFlesh/ChaosFlesh.Build.cs
new file mode 100644
--- /dev/null
+++ b/Engine/Plugins/Experimental/ChaosFlesh/Source/ChaosFlesh/ChaosFlesh.Build.cs
@@ -0,0 +1,25 @@
+// Copyright Epic Games, Inc. All Rights Reserved.
+
+namespace UnrealBuildTool.Rules
+{
+	public class ChaosFlesh : ModuleRules
+	{
+        public ChaosFlesh(ReadOnlyTargetRules Target) : base(Target)
+		{
+			PrivateIncludePaths.Add("ChaosFlesh/Private");
+            PublicIncludePaths.Add(ModuleDirectory + "/Public");
+
+			PublicDependencyModuleNames.AddRange(
+				new string[]
+				{
+					"Core",
+					"CoreUObject",
+					"Chaos",
+					"InputCore",
+				}
+			);
+
+			PrivateDefinitions.Add("CHAOS_INCLUDE_LEVEL_1=1");
+		}
+	}
+}
diff --git a/Engine/Plugins/Experimental/ChaosFlesh/Source/ChaosFlesh/Private/ChaosFlesh/ChaosFleshModule.cpp b/Engine/Plugins/Experimental/ChaosFlesh/Source/ChaosFlesh/Private/ChaosFlesh/ChaosFleshModule.cpp
new file mode 100644
--- /dev/null
+++ b/Engine/Plugins/Experimental/ChaosFlesh/Source/ChaosFlesh/Private/ChaosFlesh/ChaosFleshModule.cpp
@@ -0,0 +1,26 @@
+// Copyright Epic Games, Inc. All Rights Reserved.
+
+#include "ChaosFlesh/ChaosFleshModule.h"
+#include "Features/IModularFeatures.h"
+#include "Modules/ModuleManager.h"
+
+//////////////////////////////////////////////////////////////////////////
+// FChaosFleshModule
+
+class FChaosFleshModule : public IChaosFleshModuleInterface
+{
+  public:
+    virtual void StartupModule() override
+    {
+        check(GConfig);
+    }
+
+    virtual void ShutdownModule() override
+    {
+    }
+
+};
+
+//////////////////////////////////////////////////////////////////////////
+
+IMPLEMENT_MODULE(FChaosFleshModule, ChaosFlesh);
diff --git a/Engine/Plugins/Experimental/ChaosFlesh/Source/ChaosFlesh/Private/ChaosFlesh/FleshCollection.cpp b/Engine/Plugins/Experimental/ChaosFlesh/Source/ChaosFlesh/Private/ChaosFlesh/FleshCollection.cpp
new file mode 100644
--- /dev/null
+++ b/Engine/Plugins/Experimental/ChaosFlesh/Source/ChaosFlesh/Private/ChaosFlesh/FleshCollection.cpp
@@ -0,0 +1,46 @@
+// Copyright Epic Games, Inc. All Rights Reserved.
+
+/*=============================================================================
+	FleshCollection.cpp: FFleshCollection methods.
+=============================================================================*/
+
+#include "ChaosFlesh/FleshCollection.h"
+
+#include "ChaosFlesh/FleshCollectionUtility.h"
+
+DEFINE_LOG_CATEGORY_STATIC(FFleshCollectionLogging, Log, All);
+
+
+// Attributes
+const FName FFleshCollection::MassAttribute("Mass");
+
+FFleshCollection::FFleshCollection()
+	: FTetrahedralCollection()
+{
+	Construct();
+}
+
+
+void FFleshCollection::Construct()
+{
+	// Vertices Group
+	AddExternalAttribute<float>(FFleshCollection::MassAttribute, FFleshCollection::VerticesGroup, Mass);
+}
+
+FFleshCollection* FFleshCollection::NewFleshCollection(const TArray<FVector3f>& Vertices, const TArray<FIntVector3>& SurfaceElements, const TArray<FIntVector4>& Elements, bool bReverseVertexOrder)
+{
+
+	FFleshCollection* Collection = new FFleshCollection();
+	FFleshCollection::Init(Collection, Vertices,SurfaceElements, Elements, bReverseVertexOrder);
+	return Collection;
+}
+
+void FFleshCollection::Init(FFleshCollection* Collection, const TArray<FVector3f>& Vertices, const TArray<FIntVector3>& SurfaceElements, const TArray<FIntVector4>& Elements, bool bReverseVertexOrder)
+{
+	if (Collection)
+	{
+		FTetrahedralCollection::Init(Collection, Vertices, SurfaceElements, Elements, bReverseVertexOrder);
+	}
+}
+
+
diff --git a/Engine/Plugins/Experimental/ChaosFlesh/Source/ChaosFlesh/Private/ChaosFlesh/FleshCollectionUtility.cpp b/Engine/Plugins/Experimental/ChaosFlesh/Source/ChaosFlesh/Private/ChaosFlesh/FleshCollectionUtility.cpp
new file mode 100644
--- /dev/null
+++ b/Engine/Plugins/Experimental/ChaosFlesh/Source/ChaosFlesh/Private/ChaosFlesh/FleshCollectionUtility.cpp
@@ -0,0 +1,164 @@
+// Copyright Epic Games, Inc. All Rights Reserved.
+
+/*=============================================================================
+	GeometryCollection->cpp: FGeometryCollection methods.
+=============================================================================*/
+
+#include "ChaosFlesh/FleshCollectionUtility.h"
+
+#include "ChaosFlesh/FleshCollection.h"
+#include "CoreMinimal.h"
+#include "Chaos/Real.h"
+#include <fstream>
+#include <string>
+
+// logging
+//#include "Logging\LogMacros.h"
+DEFINE_LOG_CATEGORY(LogChaosFlesh);
+
+
+namespace ChaosFlesh
+{
+
+	TUniquePtr<FFleshCollection> ImportTetFromFile(const FString& Filename)
+	{
+		UE_LOG(LogChaosFlesh, Display, TEXT("Reading Path %s"), *Filename);
+
+		TUniquePtr<FFleshCollection> Collection;
+
+		TArray<FVector3f> Vertices;
+		TArray<FIntVector4 > Elements;
+		TArray<FIntVector3> SurfaceElements;
+
+		std::ifstream File(std::string(TCHAR_TO_UTF8(*Filename)));
+		{
+			int NumParticles;
+			{
+				std::string Input;
+				std::stringstream ss;
+				if (!ensureMsgf(getline(File, Input, ','), TEXT("Failed to read tet file. (%s)"), *Filename))
+				{
+					return Collection;
+				}
+				ss << Input << " ";
+				ss >> NumParticles;
+			}
+			Vertices.SetNum(NumParticles);
+			for (int i = 0; i < NumParticles; ++i)
+			{
+				Chaos::FReal x, y, z;
+				std::string Input;
+				std::stringstream ss;
+				if (!ensureMsgf(getline(File, Input, ','), TEXT("Failed to read tet file. (%s)"), *Filename))
+				{
+					return Collection;
+				}
+				ss << Input << " ";
+				if (!ensureMsgf(getline(File, Input, ','), TEXT("Failed to read tet file. (%s)"), *Filename))
+				{
+					return Collection;
+				}
+				ss << Input << " ";
+				if (!ensureMsgf(getline(File, Input, ','), TEXT("Failed to read tet file. (%s)"), *Filename))
+				{
+					return Collection;
+				}
+				ss << Input << " ";
+				ss >> x;
+				ss >> y;
+				ss >> z;
+				Vertices[i] = (FVector3f(x, y, z) + FVector3f(0, 4, 0));
+			}
+		}
+		{
+			int NumElements;
+			{
+				std::string Input;
+				std::stringstream ss;
+				if (!ensureMsgf(getline(File, Input, ','), TEXT("Failed to read tet file. (%s)"), *Filename))
+				{
+					return Collection;
+				}
+				ss << Input << " ";
+				ss >> NumElements;
+			}
+			Elements.SetNum(NumElements);
+			for (int i = 0; i < NumElements; ++i)
+			{
+				int x, y, z, w;
+				std::string Input;
+				std::stringstream ss;
+				if (!ensureMsgf(getline(File, Input, ','), TEXT("Failed to read tet file. (%s)"), *Filename))
+				{
+					return Collection;
+				}
+				ss << Input << " ";
+				if (!ensureMsgf(getline(File, Input, ','), TEXT("Failed to read tet file. (%s)"), *Filename))
+				{
+					return Collection;
+				}
+				ss << Input << " ";
+				if (!ensureMsgf(getline(File, Input, ','), TEXT("Failed to read tet file. (%s)"), *Filename))
+				{
+					return Collection;
+				}
+				ss << Input << " ";
+				if (!ensureMsgf(getline(File, Input, ','), TEXT("Failed to read tet file. (%s)"), *Filename))
+				{
+					return Collection;
+				}
+				ss << Input << " ";
+				ss >> x;
+				ss >> y;
+				ss >> z;
+				ss >> w;
+				Elements[i] = FIntVector4(x, y, z, w);
+			}
+		}
+		{
+			int NumElements;
+			{
+				std::string Input;
+				std::stringstream ss;
+				if (!ensureMsgf(getline(File, Input, ','), TEXT("Failed to read tet file. (%s)"), *Filename))
+				{
+					return Collection;
+				}
+				ss << Input << " ";
+				ss >> NumElements;
+			}
+			SurfaceElements.SetNum(NumElements);
+			for (int i = 0; i < NumElements; ++i)
+			{
+				int x, y, z;
+				std::string Input;
+				std::stringstream ss;
+				if (!ensureMsgf(getline(File, Input, ','), TEXT("Failed to read tet file. (%s)"), *Filename))
+				{
+					return Collection;
+				}
+				ss << Input << " ";
+				if (!ensureMsgf(getline(File, Input, ','), TEXT("Failed to read tet file. (%s)"), *Filename))
+				{
+					return Collection;
+				}
+				ss << Input << " ";
+				if (!ensureMsgf(getline(File, Input, ','), TEXT("Failed to read tet file. (%s)"), *Filename))
+				{
+					return Collection;
+				}
+				ss << Input << " ";
+				ss >> x;
+				ss >> y;
+				ss >> z;
+				SurfaceElements[i] = FIntVector3(x, y, z);
+			}
+		}
+
+		Collection.Reset(FFleshCollection::NewFleshCollection(Vertices, SurfaceElements, Elements));
+		return Collection;
+		
+
+	}
+
+}
\ No newline at end of file
diff --git a/Engine/Plugins/Experimental/ChaosFlesh/Source/ChaosFlesh/Private/ChaosFlesh/TetrahedralCollection.cpp b/Engine/Plugins/Experimental/ChaosFlesh/Source/ChaosFlesh/Private/ChaosFlesh/TetrahedralCollection.cpp
new file mode 100644
--- /dev/null
+++ b/Engine/Plugins/Experimental/ChaosFlesh/Source/ChaosFlesh/Private/ChaosFlesh/TetrahedralCollection.cpp
@@ -0,0 +1,80 @@
+// Copyright Epic Games, Inc. All Rights Reserved.
+#include "ChaosFlesh/TetrahedralCollection.h"
+
+#include "Chaos/ChaosArchive.h"
+
+DEFINE_LOG_CATEGORY_STATIC(FTetrahedralCollectionLogging, Log, All);
+
+// Groups
+const FName FTetrahedralCollection::TetrahedralGroup = "Tetrahedral";
+
+// Attributes
+const FName FTetrahedralCollection::TetrahedronAttribute("Tetrahedron");
+
+FTetrahedralCollection::FTetrahedralCollection()
+	: Super::FGeometryCollection()
+{
+	Construct();
+}
+
+
+void FTetrahedralCollection::Construct()
+{
+	FManagedArrayCollection::FConstructionParameters TetrahedronDependency(FTetrahedralCollection::TetrahedralGroup);
+
+	// Tetrahedron Group
+	AddExternalAttribute<FIntVector4>(FTetrahedralCollection::TetrahedralGroup, FTetrahedralCollection::TetrahedralGroup, Tetrahedron);
+}
+
+
+void FTetrahedralCollection::SetDefaults(FName Group, uint32 StartSize, uint32 NumElements)
+{
+	Super::SetDefaults(Group, StartSize, NumElements);
+
+	if (Group == FTetrahedralCollection::TetrahedralGroup)
+	{
+		for (uint32 Idx = StartSize; Idx < StartSize + NumElements; ++Idx)
+		{
+			Tetrahedron[Idx] = FIntVector4(INDEX_NONE);
+		}
+	}
+}
+
+FTetrahedralCollection* FTetrahedralCollection::NewTetrahedralCollection(const TArray<FVector3f>& Vertices, const TArray<FIntVector3>& SurfaceElements, const TArray<FIntVector4>& Elements, bool bReverseVertexOrder)
+{
+	FTetrahedralCollection* Collection = new FTetrahedralCollection();
+	FTetrahedralCollection::Init(Collection, Vertices, SurfaceElements, Elements, bReverseVertexOrder);
+	return Collection;
+}
+void FTetrahedralCollection::Init(FTetrahedralCollection* Collection, const TArray<FVector3f>& Vertices, const TArray<FIntVector3>& SurfaceElements, const TArray<FIntVector4>& Elements, bool bReverseVertexOrder)
+{
+	if (Collection)
+	{
+		TArray<float> RawVertexArray;
+		RawVertexArray.SetNum(Vertices.Num()*3);
+		for(int i=0,j=0;i<Vertices.Num();i++,j+=3)
+		{
+			RawVertexArray[j+0] = Vertices[i].X;
+			RawVertexArray[j+1] = Vertices[i].Y;
+			RawVertexArray[j+2] = Vertices[i].Z;
+		}
+		TArray<int32> RawIndicesArray;
+		RawIndicesArray.SetNum(SurfaceElements.Num()*3);
+		for (int i = 0, j = 0; i < Vertices.Num(); i++, j += 3)
+		{
+			RawIndicesArray[j + 0] = SurfaceElements[i].X;
+			RawIndicesArray[j + 1] = SurfaceElements[i].Y;
+			RawIndicesArray[j + 2] = SurfaceElements[i].Z;
+		}
+
+		Super::Init(Collection, RawVertexArray, RawIndicesArray, bReverseVertexOrder);
+
+		Collection->AddElements(Elements.Num(), FTetrahedralCollection::TetrahedralGroup);
+		for (int i = 0; i < Elements.Num(); i++)
+		{
+			Collection->Tetrahedron[i] = Elements[i];
+		}
+	}
+}
+
+
diff --git a/Engine/Plugins/Experimental/ChaosFlesh/Source/ChaosFlesh/Public/ChaosFlesh/ChaosFleshModule.h b/Engine/Plugins/Experimental/ChaosFlesh/Source/ChaosFlesh/Public/ChaosFlesh/ChaosFleshModule.h
new file mode 100644
--- /dev/null
+++ b/Engine/Plugins/Experimental/ChaosFlesh/Source/ChaosFlesh/Public/ChaosFlesh/ChaosFleshModule.h
@@ -0,0 +1,15 @@
+// Copyright Epic Games, Inc. All Rights Reserved.
+
+#pragma once
+
+#include "CoreMinimal.h"
+#include "Modules/ModuleInterface.h"
+
+#define	CHAOS_FLESH_MODULE_NAME "ChaosFlesh"
+
+//////////////////////////////////////////////////////////////////////////
+// IChaosFleshModuleInterface
+
+class IChaosFleshModuleInterface : public IModuleInterface
+{
+};
diff --git a/Engine/Plugins/Experimental/ChaosFlesh/Source/ChaosFlesh/Public/ChaosFlesh/FleshCollection.h b/Engine/Plugins/Experimental/ChaosFlesh/Source/ChaosFlesh/Public/ChaosFlesh/FleshCollection.h
new file mode 100644
--- /dev/null
+++ b/Engine/Plugins/Experimental/ChaosFlesh/Source/ChaosFlesh/Public/ChaosFlesh/FleshCollection.h
@@ -0,0 +1,49 @@
+// Copyright Epic Games, Inc. All Rights Reserved.
+
+#pragma once
+
+#include "CoreMinimal.h"
+#include "ChaosFlesh/TetrahedralCollection.h"
+
+namespace Chaos
+{
+	class FChaosArchive;
+}
+
+/**
+* FFleshCollection (FTetrehidralCollection)
+*/
+class CHAOSFLESH_API FFleshCollection : public FTetrahedralCollection
+{
+public:
+	typedef FTetrahedralCollection Super;
+
+	FFleshCollection();
+	FFleshCollection(FFleshCollection&) = delete;
+	FFleshCollection& operator=(const FFleshCollection&) = delete;
+	FFleshCollection(FFleshCollection&&) = default;
+	FFleshCollection& operator=(FFleshCollection&&) = default;
+
+
+	static const FName MassAttribute; // (Mass,Vertices)
+
+	/**
+	 * Create a GeometryCollection from Vertex and Indices arrays
+	 */
+	static FFleshCollection* NewFleshCollection(const TArray<FVector3f>& Vertices, const TArray<FIntVector3>& SurfaceElements, const TArray<FIntVector4>& Elements, bool bReverseVertexOrder = true);
+	static void Init(FFleshCollection* Collection, const TArray<FVector3f>& Vertices, const TArray<FIntVector3>& SurfaceElements, const TArray<FIntVector4>& Elements,  bool bReverseVertexOrder = true);
+
+
+	// Simualtion attributes (Vertices Group)
+	TManagedArray<float>		Mass;
+
+protected:
+	void Construct();
+
+};
+
+FORCEINLINE Chaos::FChaosArchive& operator<<(Chaos::FChaosArchive& Ar, FFleshCollection& Value)
+{
+	Value.Serialize(Ar);
+	return Ar;
+}
diff --git a/Engine/Plugins/Experimental/ChaosFlesh/Source/ChaosFlesh/Public/ChaosFlesh/FleshCollectionUtility.h b/Engine/Plugins/Experimental/ChaosFlesh/Source/ChaosFlesh/Public/ChaosFlesh/FleshCollectionUtility.h
new file mode 100644
--- /dev/null
+++ b/Engine/Plugins/Experimental/ChaosFlesh/Source/ChaosFlesh/Public/ChaosFlesh/FleshCollectionUtility.h
@@ -0,0 +1,15 @@
+// Copyright Epic Games, Inc. All Rights Reserved.
+
+#pragma once
+
+#include "ChaosFlesh/FleshCollection.h"
+#include "CoreMinimal.h"
+#include "Templates/UniquePtr.h"
+#include "Misc/AssertionMacros.h"
+
+DECLARE_LOG_CATEGORY_EXTERN(LogChaosFlesh, Verbose, All);
+
+namespace ChaosFlesh 
+{
+	TUniquePtr<FFleshCollection> CHAOSFLESH_API ImportTetFromFile(const FString& Filename);
+}
diff --git a/Engine/Plugins/Experimental/ChaosFlesh/Source/ChaosFlesh/Public/ChaosFlesh/TetrahedralCollection.h b/Engine/Plugins/Experimental/ChaosFlesh/Source/ChaosFlesh/Public/ChaosFlesh/TetrahedralCollection.h
new file mode 100644
--- /dev/null
+++ b/Engine/Plugins/Experimental/ChaosFlesh/Source/ChaosFlesh/Public/ChaosFlesh/TetrahedralCollection.h
@@ -0,0 +1,59 @@
+// Copyright Epic Games, Inc. All Rights Reserved.
+
+#pragma once
+
+#include "CoreMinimal.h"
+#include "GeometryCollection/GeometryCollection.h"
+#include "GeometryCollection/ManagedArrayCollection.h"
+#include "Misc/Crc.h"
+
+namespace Chaos
+{
+	class FChaosArchive;
+}
+
+
+/**
+* FTetrahedralCollection (FGeometryCollection)
+*/
+class CHAOSFLESH_API FTetrahedralCollection : public FGeometryCollection
+{
+	typedef FGeometryCollection Super;
+
+public:
+	FTetrahedralCollection();
+	FTetrahedralCollection(FTetrahedralCollection &) = delete;
+	FTetrahedralCollection& operator=(const FTetrahedralCollection &) = delete;
+	FTetrahedralCollection(FTetrahedralCollection &&) = default;
+	FTetrahedralCollection& operator=(FTetrahedralCollection &&) = default;
+
+
+	// groups
+	static const FName TetrahedralGroup;
+
+	// attributes
+	static const FName TetrahedronAttribute;
+
+	/**
+	 * Create a GeometryCollection from Vertex and Indices arrays
+	 */
+	static FTetrahedralCollection* NewTetrahedralCollection(const TArray<FVector3f>& Vertices, const TArray<FIntVector3>& SurfaceElements, const TArray<FIntVector4>& Elements,bool bReverseVertexOrder = true);
+	static void Init(FTetrahedralCollection* Collection, const TArray<FVector3f>& Vertices, const TArray<FIntVector3>& SurfaceElements, const TArray<FIntVector4>& Elements, bool bReverseVertexOrder = true);
+
+	void SetDefaults(FName Group, uint32 StartSize, uint32 NumElements);
+
+	// Transform Group
+	TManagedArray<FIntVector4>	Tetrahedron;
+	
+protected:
+
+	void Construct();
+
+};
+
+FORCEINLINE Chaos::FChaosArchive& operator<<(Chaos::FChaosArchive& Ar, FTetrahedralCollection& Value)
+{
+	Value.Serialize(Ar);
+	return Ar;
+}
+
