Chaos : Flesh
- Updates to initialization.

diff --git a/Engine/Plugins/Experimental/ChaosFlesh/Source/ChaosFlesh/Private/ChaosFlesh/FleshCollection.cpp b/Engine/Plugins/Experimental/ChaosFlesh/Source/ChaosFlesh/Private/ChaosFlesh/FleshCollection.cpp
--- a/Engine/Plugins/Experimental/ChaosFlesh/Source/ChaosFlesh/Private/ChaosFlesh/FleshCollection.cpp
+++ b/Engine/Plugins/Experimental/ChaosFlesh/Source/ChaosFlesh/Private/ChaosFlesh/FleshCollection.cpp
@@ -25,17 +25,26 @@ void FFleshCollection::Construct()
 {
 	// Vertices Group
 	AddExternalAttribute<float>(FFleshCollection::MassAttribute, FFleshCollection::VerticesGroup, Mass);
+	for (float& m : Mass) { m = 1; }
 }
 
-FFleshCollection* FFleshCollection::NewFleshCollection(const TArray<FVector3f>& Vertices, const TArray<FIntVector3>& SurfaceElements, const TArray<FIntVector4>& Elements, bool bReverseVertexOrder)
+
+FFleshCollection* FFleshCollection::NewFleshCollection(const TArray<FVector>& Vertices, const TArray<FIntVector3>& SurfaceElements, const TArray<FIntVector4>& Elements, bool bReverseVertexOrder)
 {
 
 	FFleshCollection* Collection = new FFleshCollection();
 	FFleshCollection::Init(Collection, Vertices,SurfaceElements, Elements, bReverseVertexOrder);
 	return Collection;
 }
+FFleshCollection* FFleshCollection::NewFleshCollection(const FTetrahedralCollection& Base)
+{
+	FFleshCollection* Collection = new FFleshCollection();
+	Collection->CopyMatchingAttributesFrom(Base);
+	for (float& m : Collection->Mass) { m = 1; }
+	return Collection;
+}
 
-void FFleshCollection::Init(FFleshCollection* Collection, const TArray<FVector3f>& Vertices, const TArray<FIntVector3>& SurfaceElements, const TArray<FIntVector4>& Elements, bool bReverseVertexOrder)
+void FFleshCollection::Init(FFleshCollection* Collection, const TArray<FVector>& Vertices, const TArray<FIntVector3>& SurfaceElements, const TArray<FIntVector4>& Elements, bool bReverseVertexOrder)
 {
 	if (Collection)
 	{
diff --git a/Engine/Plugins/Experimental/ChaosFlesh/Source/ChaosFlesh/Private/ChaosFlesh/FleshCollectionUtility.cpp b/Engine/Plugins/Experimental/ChaosFlesh/Source/ChaosFlesh/Private/ChaosFlesh/FleshCollectionUtility.cpp
--- a/Engine/Plugins/Experimental/ChaosFlesh/Source/ChaosFlesh/Private/ChaosFlesh/FleshCollectionUtility.cpp
+++ b/Engine/Plugins/Experimental/ChaosFlesh/Source/ChaosFlesh/Private/ChaosFlesh/FleshCollectionUtility.cpp
@@ -26,7 +26,7 @@ namespace ChaosFlesh
 
 		TUniquePtr<FFleshCollection> Collection;
 
-		TArray<FVector3f> Vertices;
+		TArray<FVector> Vertices;
 		TArray<FIntVector4 > Elements;
 		TArray<FIntVector3> SurfaceElements;
 
@@ -67,7 +67,7 @@ namespace ChaosFlesh
 				ss >> x;
 				ss >> y;
 				ss >> z;
-				Vertices[i] = (FVector3f(x, y, z) + FVector3f(0, 4, 0));
+				Vertices[i] = (FVector(x, y, z) + FVector(0, 4, 0));
 			}
 		}
 		{
diff --git a/Engine/Plugins/Experimental/ChaosFlesh/Source/ChaosFlesh/Private/ChaosFlesh/TetrahedralCollection.cpp b/Engine/Plugins/Experimental/ChaosFlesh/Source/ChaosFlesh/Private/ChaosFlesh/TetrahedralCollection.cpp
--- a/Engine/Plugins/Experimental/ChaosFlesh/Source/ChaosFlesh/Private/ChaosFlesh/TetrahedralCollection.cpp
+++ b/Engine/Plugins/Experimental/ChaosFlesh/Source/ChaosFlesh/Private/ChaosFlesh/TetrahedralCollection.cpp
@@ -40,13 +40,13 @@ void FTetrahedralCollection::SetDefaults(FName Group, uint32 StartSize, uint32 N
 	}
 }
 
-FTetrahedralCollection* FTetrahedralCollection::NewTetrahedralCollection(const TArray<FVector3f>& Vertices, const TArray<FIntVector3>& SurfaceElements, const TArray<FIntVector4>& Elements, bool bReverseVertexOrder)
+FTetrahedralCollection* FTetrahedralCollection::NewTetrahedralCollection(const TArray<FVector>& Vertices, const TArray<FIntVector3>& SurfaceElements, const TArray<FIntVector4>& Elements, bool bReverseVertexOrder)
 {
 	FTetrahedralCollection* Collection = new FTetrahedralCollection();
 	FTetrahedralCollection::Init(Collection, Vertices, SurfaceElements, Elements, bReverseVertexOrder);
 	return Collection;
 }
-void FTetrahedralCollection::Init(FTetrahedralCollection* Collection, const TArray<FVector3f>& Vertices, const TArray<FIntVector3>& SurfaceElements, const TArray<FIntVector4>& Elements, bool bReverseVertexOrder)
+void FTetrahedralCollection::Init(FTetrahedralCollection* Collection, const TArray<FVector>& Vertices, const TArray<FIntVector3>& SurfaceElements, const TArray<FIntVector4>& Elements, bool bReverseVertexOrder)
 {
 	if (Collection)
 	{
diff --git a/Engine/Plugins/Experimental/ChaosFlesh/Source/ChaosFlesh/Private/ChaosFlesh/TetrahedralExampleGeometry.cpp b/Engine/Plugins/Experimental/ChaosFlesh/Source/ChaosFlesh/Private/ChaosFlesh/TetrahedralExampleGeometry.cpp
new file mode 100644
--- /dev/null
+++ b/Engine/Plugins/Experimental/ChaosFlesh/Source/ChaosFlesh/Private/ChaosFlesh/TetrahedralExampleGeometry.cpp
@@ -0,0 +1,55 @@
+// Copyright Epic Games, Inc. All Rights Reserved.
+
+/*=============================================================================
+	GeometryCollection->cpp: FGeometryCollection methods.
+=============================================================================*/
+
+#include "ChaosFlesh/TetrahedralExampleGeometry.h"
+
+#include "ChaosFlesh/FleshCollection.h"
+#include "Chaos/Real.h"
+
+
+namespace ChaosFlesh
+{
+	namespace ExampleGeometry
+	{
+		FVector TetrahedralCentroid(TArray<FVector> V)
+		{
+			return FVector(
+				(V[0].X + V[1].X + V[2].X + V[3].X) / 4.0,
+				(V[0].Y + V[1].Y + V[2].Y + V[3].Y) / 4.0,
+				(V[0].Z + V[1].Z + V[2].Z + V[3].Z) / 4.0
+			);
+		}
+
+		TUniquePtr<FTetrahedralCollection> UnitTetrahedron::Create()
+		{
+			float unitProj = FMath::Sin(PI / 3.0); // projection of unit vector at 60 deg
+			float faceAngle = FMath::Acos(1. / 3); // angle between planes of a tetrahedron. 
+			TArray<FVector> Vertices = {
+				FVector(-0.5, 0.0, 0.0),
+				FVector( 0.5, 0.0, 0.0),
+				FVector( 0.0, unitProj, 0.0),
+				FVector( 0.0, unitProj * FMath::Cos(faceAngle), unitProj * FMath::Sin(faceAngle)),
+			};
+			FVector Center = TetrahedralCentroid(Vertices);
+			for (FVector & Vert : Vertices)
+			{
+				Vert -= Center;
+			}
+
+			TArray<FIntVector> Surface = {
+				FIntVector(0,1,2),
+				FIntVector(0,2,3),
+				FIntVector(3,2,1),
+				FIntVector(0,3,1)
+			};
+			TArray<FIntVector4> Tetrahedron = {
+					FIntVector4(0,1,2,3)
+			};
+			return 	TUniquePtr<FFleshCollection>(FFleshCollection::NewFleshCollection(Vertices, Surface, Tetrahedron, false));
+			;
+		}
+	}
+}
\ No newline at end of file
diff --git a/Engine/Plugins/Experimental/ChaosFlesh/Source/ChaosFlesh/Public/ChaosFlesh/FleshCollection.h b/Engine/Plugins/Experimental/ChaosFlesh/Source/ChaosFlesh/Public/ChaosFlesh/FleshCollection.h
--- a/Engine/Plugins/Experimental/ChaosFlesh/Source/ChaosFlesh/Public/ChaosFlesh/FleshCollection.h
+++ b/Engine/Plugins/Experimental/ChaosFlesh/Source/ChaosFlesh/Public/ChaosFlesh/FleshCollection.h
@@ -18,24 +18,27 @@ class CHAOSFLESH_API FFleshCollection : public FTetrahedralCollection
 public:
 	typedef FTetrahedralCollection Super;
 
+
 	FFleshCollection();
 	FFleshCollection(FFleshCollection&) = delete;
 	FFleshCollection& operator=(const FFleshCollection&) = delete;
 	FFleshCollection(FFleshCollection&&) = default;
 	FFleshCollection& operator=(FFleshCollection&&) = default;
 
-
-	static const FName MassAttribute; // (Mass,Vertices)
-
 	/**
 	 * Create a GeometryCollection from Vertex and Indices arrays
 	 */
-	static FFleshCollection* NewFleshCollection(const TArray<FVector3f>& Vertices, const TArray<FIntVector3>& SurfaceElements, const TArray<FIntVector4>& Elements, bool bReverseVertexOrder = true);
-	static void Init(FFleshCollection* Collection, const TArray<FVector3f>& Vertices, const TArray<FIntVector3>& SurfaceElements, const TArray<FIntVector4>& Elements,  bool bReverseVertexOrder = true);
+	static FFleshCollection* NewFleshCollection(const FTetrahedralCollection& Base);
+	static FFleshCollection* NewFleshCollection(const TArray<FVector>& Vertices, const TArray<FIntVector3>& SurfaceElements, const TArray<FIntVector4>& Elements, bool bReverseVertexOrder = true);
+	static void Init(FFleshCollection* Collection, const TArray<FVector>& Vertices, const TArray<FIntVector3>& SurfaceElements, const TArray<FIntVector4>& Elements, bool bReverseVertexOrder = true);
 
 
-	// Simualtion attributes (Vertices Group)
-	TManagedArray<float>		Mass;
+	/*
+	* Mass Attribute
+	*   FManagedArray<float> Mass = this->FindAttribute<float>("Mass",FGeometryCollection::VerticesGroup);
+	*/
+	static const FName MassAttribute;
+	TManagedArray<float> Mass;
 
 protected:
 	void Construct();
diff --git a/Engine/Plugins/Experimental/ChaosFlesh/Source/ChaosFlesh/Public/ChaosFlesh/TetrahedralCollection.h b/Engine/Plugins/Experimental/ChaosFlesh/Source/ChaosFlesh/Public/ChaosFlesh/TetrahedralCollection.h
--- a/Engine/Plugins/Experimental/ChaosFlesh/Source/ChaosFlesh/Public/ChaosFlesh/TetrahedralCollection.h
+++ b/Engine/Plugins/Experimental/ChaosFlesh/Source/ChaosFlesh/Public/ChaosFlesh/TetrahedralCollection.h
@@ -27,23 +27,28 @@ public:
 	FTetrahedralCollection(FTetrahedralCollection &&) = default;
 	FTetrahedralCollection& operator=(FTetrahedralCollection &&) = default;
 
-
-	// groups
-	static const FName TetrahedralGroup;
-
-	// attributes
-	static const FName TetrahedronAttribute;
-
 	/**
 	 * Create a GeometryCollection from Vertex and Indices arrays
 	 */
-	static FTetrahedralCollection* NewTetrahedralCollection(const TArray<FVector3f>& Vertices, const TArray<FIntVector3>& SurfaceElements, const TArray<FIntVector4>& Elements,bool bReverseVertexOrder = true);
-	static void Init(FTetrahedralCollection* Collection, const TArray<FVector3f>& Vertices, const TArray<FIntVector3>& SurfaceElements, const TArray<FIntVector4>& Elements, bool bReverseVertexOrder = true);
+	static FTetrahedralCollection* NewTetrahedralCollection(const TArray<FVector>& Vertices, const TArray<FIntVector3>& SurfaceElements, const TArray<FIntVector4>& Elements,bool bReverseVertexOrder = true);
+	static void Init(FTetrahedralCollection* Collection, const TArray<FVector>& Vertices, const TArray<FIntVector3>& SurfaceElements, const TArray<FIntVector4>& Elements, bool bReverseVertexOrder = true);
+
+	/*
+	*  SetDefaults for new entries on this collection. 
+	*/
+	virtual void SetDefaults(FName Group, uint32 StartSize, uint32 NumElements);
 
-	void SetDefaults(FName Group, uint32 StartSize, uint32 NumElements);
+	/*
+	*  Attribute Groupes
+	*/
+	static const FName TetrahedralGroup;
 
-	// Transform Group
-	TManagedArray<FIntVector4>	Tetrahedron;
+	/*
+	*  Tetrahedron Attribute
+	*  TManagedArray<FIntVector4> Tetrahedron = this->FindAttribute<FIntVector4>(FTetrahedralCollection::TetrahedronAttribute,FTetrahedralCollection::TetrahedralGroup);
+	*/
+	static const FName TetrahedronAttribute;
+	TManagedArray<FIntVector4>Tetrahedron;
 	
 protected:
 
diff --git a/Engine/Plugins/Experimental/ChaosFlesh/Source/ChaosFlesh/Public/ChaosFlesh/TetrahedralExampleGeometry.h b/Engine/Plugins/Experimental/ChaosFlesh/Source/ChaosFlesh/Public/ChaosFlesh/TetrahedralExampleGeometry.h
new file mode 100644
--- /dev/null
+++ b/Engine/Plugins/Experimental/ChaosFlesh/Source/ChaosFlesh/Public/ChaosFlesh/TetrahedralExampleGeometry.h
@@ -0,0 +1,18 @@
+// Copyright Epic Games, Inc. All Rights Reserved.
+
+#pragma once
+
+#include "ChaosFlesh/TetrahedralCollection.h"
+#include "Templates/UniquePtr.h"
+#include "CoreFwd.h"
+
+namespace ChaosFlesh 
+{
+	namespace ExampleGeometry
+	{
+		struct UnitTetrahedron 
+		{
+			static TUniquePtr<FTetrahedralCollection> CHAOSFLESH_API Create();
+		};
+	}
+}
