ChaosFlesh : Moved to Experimental

diff --git a/Engine/Plugins/Experimental/ChaosFlesh/Source/ChaosFlesh/ChaosFlesh.Build.cs b/Engine/Plugins/Experimental/ChaosFlesh/Source/ChaosFlesh/ChaosFlesh.Build.cs
new file mode 100644
--- /dev/null
+++ b/Engine/Plugins/Experimental/ChaosFlesh/Source/ChaosFlesh/ChaosFlesh.Build.cs
@@ -0,0 +1,57 @@
+// Copyright Epic Games, Inc. All Rights Reserved.
+
+namespace UnrealBuildTool.Rules
+{
+	
+	public class ChaosFlesh : ModuleRules
+	{
+		public ChaosFlesh(ReadOnlyTargetRules Target) : base(Target)
+		{
+			// We can do specializations for PUBLIC platforms in this file; otherwise it
+			// may not be defined in any given workspace, and references to it will fail.
+			// You can do a string comparison ala:
+			// \code
+			//		Target.Platform.ToString() == "XXX"
+			// \endcode
+			// But that may break our NDA policy which forbids even mentioning NDA'd
+			// platforms in engine code.  A less fragile way to go about it would be to
+			// add a platform specific build file:
+			//
+			// ChaosFlesh_XXX.Build.cs:
+			// \code
+			//    namespace UnrealBuildTool.Rules
+			//    {
+			//        public class ChaosFlesh_XXX : ChaosFlesh
+			//        {
+			//            public ChaosFlesh_XXX(ReadOnlyTargetRules Target) : base(Target) 
+			//            { <Do XXX specific dances here> }
+			//        }
+			//    }
+			// \endcode
+			//
+			// However, adding one of these build files for a public platform will clash with
+			// specializations that the build system makes.  In practice, if it's in
+			// Engine\Source\Programs\UnrealBuildTool\Platform it is public, otherwise it's not.
+			if (Target.Platform == UnrealTargetPlatform.Win64 ||
+				Target.Platform == UnrealTargetPlatform.Linux ||
+				Target.Platform == UnrealTargetPlatform.Mac)
+			{
+				AddEngineThirdPartyPrivateStaticDependencies(Target, "zlib");
+				PrivateDefinitions.Add("USE_ZLIB=1");
+			}
+
+			PublicDependencyModuleNames.AddRange(
+				new string[]
+				{
+					"Core",
+					"CoreUObject",
+					"Chaos",
+					"InputCore",
+				}
+			);
+
+			PrivateDefinitions.Add("CHAOS_INCLUDE_LEVEL_1=1");
+		}
+	}
+
+}
diff --git a/Engine/Plugins/Experimental/ChaosFlesh/Source/ChaosFlesh/Private/ChaosFlesh/FleshCollectionUtility.cpp b/Engine/Plugins/Experimental/ChaosFlesh/Source/ChaosFlesh/Private/ChaosFlesh/FleshCollectionUtility.cpp
new file mode 100644
--- /dev/null
+++ b/Engine/Plugins/Experimental/ChaosFlesh/Source/ChaosFlesh/Private/ChaosFlesh/FleshCollectionUtility.cpp
@@ -0,0 +1,352 @@
+// Copyright Epic Games, Inc. All Rights Reserved.
+
+/*=============================================================================
+	GeometryCollection->cpp: FGeometryCollection methods.
+=============================================================================*/
+
+#include "ChaosFlesh/FleshCollectionUtility.h"
+
+#include "ChaosFlesh/FleshCollection.h"
+#include "ChaosFlesh/GEO.h"
+#include "ChaosFlesh/PB.h"
+
+#include "CoreMinimal.h"
+#include "Chaos/Real.h"
+#include "Containers/Array.h"
+#include "Containers/Map.h"
+
+#include "Misc/Paths.h"
+
+#include <fstream>
+#include <string>
+
+// logging
+//#include "Logging\LogMacros.h"
+DEFINE_LOG_CATEGORY(LogChaosFlesh);
+
+namespace ChaosFlesh
+{
+	void GetTetFaces(
+		const FIntVector4& Tet,
+		FIntVector3& Face1,
+		FIntVector3& Face2,
+		FIntVector3& Face3,
+		FIntVector3& Face4,
+		const bool invert)
+	{
+		const int32 i = Tet[0];
+		const int32 j = Tet[1];
+		const int32 k = Tet[2];
+		const int32 l = Tet[3];
+		if (invert)
+		{
+			Face1 = { i, k, j };
+			Face2 = { i, j, l };
+			Face3 = { i, l, k };
+			Face4 = { j, k, l };
+		}
+		else
+		{
+			Face1 = { i, j, k };
+			Face2 = { i, l, j };
+			Face3 = { i, k, l };
+			Face4 = { j, l, k };
+		}
+	}
+
+	int32 
+	GetMin(const FIntVector3& V) 
+	{
+		return FMath::Min3(V[0], V[1], V[2]);
+	}
+	int32 
+	GetMid(const FIntVector3& V)
+	{
+		const int32 X = V[0]; const int32 Y = V[1]; const int32 Z = V[2];
+		const int32 XmY = X - Y;
+		const int32 YmZ = Y - Z;
+		const int32 XmZ = X - Z;
+		return (XmY * YmZ > -1 ? Y : XmY * XmZ < 1 ? X : Z);
+	}
+	int32 
+	GetMax(const FIntVector3& V)
+	{
+		return FMath::Max3(V[0], V[1], V[2]);
+	}
+	FIntVector3
+	GetOrdered(const FIntVector3& V)
+	{
+		return FIntVector3(GetMin(V), GetMid(V), GetMax(V));
+	}
+	FIntVector4
+	GetOrdered(const FIntVector4& V)
+	{
+		TArray<int32> VA = { V[0], V[1], V[2], V[3] };
+		VA.Sort();
+		return FIntVector4(VA[0], VA[1], VA[2], VA[3]);
+	}
+
+	void 
+	GetSurfaceElements(
+		const TArray<FIntVector4>& Tets,
+		TArray<FIntVector3>& SurfaceElements,
+		const bool KeepInteriorFaces,
+		const bool InvertFaces)
+	{
+		FIntVector3 Faces[4];
+
+		if (KeepInteriorFaces)
+		{
+			SurfaceElements.Reserve(Tets.Num() * 4);
+			for (const FIntVector4& Tet : Tets)
+			{
+				GetTetFaces(Tet, Faces[0], Faces[1], Faces[2], Faces[3], InvertFaces);
+				for (int i = 0; i < 4; i++)
+					SurfaceElements.Add(Faces[i]);
+			}
+		}
+		else
+		{
+			typedef TMap<int32, TPair<uint8, FIntVector3>> ZToCount;
+			typedef TMap<int32, ZToCount> YToZ;
+			typedef TMap<int32, YToZ> XToY;
+			XToY CoordToCount;
+
+			int32 Idx = -1;
+			int32 Count = 0;
+			for (const FIntVector4& Tet : Tets)
+			{
+				++Idx;
+				if (!(Tet[0] != Tet[1] &&
+					Tet[0] != Tet[2] &&
+					Tet[0] != Tet[3] &&
+					Tet[1] != Tet[2] &&
+					Tet[1] != Tet[3] &&
+					Tet[2] != Tet[3]))
+				{
+					UE_LOG(LogChaosFlesh, Display, TEXT("Skipping degenerate tet %d of %d."), Idx, Tets.Num());
+					continue;
+				}
+
+				GetTetFaces(Tet, Faces[0], Faces[1], Faces[2], Faces[3], InvertFaces);
+				for (int i = 0; i < 4; i++)
+				{
+					const FIntVector3 OFace = GetOrdered(Faces[i]);
+					check(OFace[0] <= OFace[1] && OFace[1] <= OFace[2]);
+					const int32 OA = OFace[0];
+					const int32 OB = OFace[1];
+					const int32 OC = OFace[2];
+
+					auto& zc = CoordToCount.FindOrAdd(OFace[0]).FindOrAdd(OFace[1]);
+					auto zcIt = zc.Find(OFace[2]);
+					if (zcIt == nullptr)
+					{
+						zc.Add(OFace[2], TPair<uint8,FIntVector3>((uint8)1,Faces[i]) );
+						// Increment our count of lone faces.
+						Count++;
+					}
+					else
+					{
+						zcIt->Key++;
+						// Since we're talking about shared faces, the only way we'd
+						// get a face instanced more than twice is if we had a degenerate 
+						// tet mesh.
+						Count--;
+					}
+				}
+			}
+
+			size_t nonManifold = 0;
+			SurfaceElements.Reserve(Count);
+			for(auto &xyIt : CoordToCount)
+			{
+				const YToZ& yz = xyIt.Value;
+				for(auto &yzIt : yz)
+				{
+					const ZToCount& zc = yzIt.Value;
+					for(auto &zcIt : zc)
+					{
+						const uint8 FaceCount = zcIt.Value.Key;
+						if (FaceCount == 1)
+						{
+							SurfaceElements.Add(zcIt.Value.Value);
+						}
+						else if (FaceCount > 2)
+						{
+							//UE_LOG(LogChaosFlesh, Display, TEXT("WARNING: Non-manifold tetrahedral mesh detected (face [%d, %d, %d] use count %d)."), zcIt->second.second[0], zcIt->second.second[1], zcIt->second.second[2], FaceCount);
+							nonManifold++;
+						}
+					}
+				}
+			}
+			if(nonManifold)
+				UE_LOG(LogChaosFlesh, Display, TEXT("WARNING: Encountered %d non-manifold tetrahedral mesh faces."), nonManifold);
+		}
+	}
+
+	TUniquePtr<FFleshCollection> ReadTetPBDeformableGeometryCollection(const FString& Filename)
+	{
+		TUniquePtr<FFleshCollection> Collection;
+		UE_LOG(LogChaosFlesh, Display, TEXT("Reading Path %s"), *Filename);
+
+		IO::DeformableGeometryCollectionReader Reader(Filename);
+		if (!Reader.ReadPBScene())
+			return Collection;
+
+		TArray<FVector> *Vertices = nullptr;
+		TArray<FVector> FrameVertices;
+		TArray<FIntVector4> *Elements = nullptr;
+		TArray<FIntVector3> SurfaceElements;
+
+		TArray<IO::DeformableGeometryCollectionReader::TetMesh*> TetMeshes = Reader.GetTetMeshes();
+		for (auto* It : TetMeshes)
+		{
+			Elements = reinterpret_cast<TArray<FIntVector4>*>(&It->Elements);
+			Vertices = &It->Points;
+			UE_LOG(LogChaosFlesh, Display, TEXT("Got %d tets."), Elements->Num());
+			UE_LOG(LogChaosFlesh, Display, TEXT("Got %d rest points."), Vertices->Num());
+			break;
+		}
+		TArray<FIntVector4> ElemTmp;
+		if (!Elements)
+		{
+			Elements = &ElemTmp;
+			UE_LOG(LogChaosFlesh, Display, TEXT("Got %d tets."), Elements->Num());
+		}
+
+		if (!Vertices || Vertices->Num() == 0)
+		{
+			TPair<int32, int32> FrameRange = Reader.ReadFrameRange();
+			if (!Reader.ReadPoints(FrameRange.Key, FrameVertices))
+			{
+				return Collection;
+			}
+			Vertices = &FrameVertices;
+			UE_LOG(LogChaosFlesh, Display, TEXT("Got %d points from frame %d."), Vertices->Num(), FrameRange.Key);
+		}
+
+		GetSurfaceElements(*Elements, SurfaceElements, false);
+		UE_LOG(LogChaosFlesh, Display, TEXT("Got %d tris."), SurfaceElements.Num());
+
+		Collection.Reset(FFleshCollection::NewFleshCollection(*Vertices, SurfaceElements, *Elements));
+		return Collection;
+	}
+
+
+	TUniquePtr<FFleshCollection> ReadGEOFile(const FString& Filename)
+	{
+		TUniquePtr<FFleshCollection> Collection;
+		UE_LOG(LogChaosFlesh, Display, TEXT("Reading Path %s"), *Filename);
+
+		TMap<FString, int32> intVars;
+		TMap<FString, TArray<int32>> intVectorVars;
+		TMap<FString, TArray<float>> floatVectorVars;
+		TMap<FString, TPair<TArray<std::string>, TArray<int32>>> indexedStringVars;
+		if (!ReadGEO(std::string(TCHAR_TO_UTF8(*Filename)), intVars, intVectorVars, floatVectorVars, indexedStringVars))
+		{
+			UE_LOG(LogChaosFlesh, Display, TEXT("Failed to open GEO file: '%s'."), *Filename);
+			return Collection;
+		}
+
+		TArray<FVector> Vertices;
+		TArray<FIntVector4> Elements;
+		TArray<FIntVector3> SurfaceElements;
+
+		auto fvIt = floatVectorVars.Find("position");
+		if (fvIt == nullptr)
+			fvIt = floatVectorVars.Find("P");
+		if (fvIt != nullptr)
+		{
+			const TArray<float>& coords = *fvIt;
+			Vertices.Reserve(coords.Num() / 3);
+			for (size_t i = 0; i < coords.Num(); i += 3)
+			{
+				FVector pt;
+				for (size_t j = 0; j < 3; j++)
+					pt[j] = coords[i + j];
+				Vertices.Add(pt);
+			}
+		}
+		UE_LOG(LogChaosFlesh, Display, TEXT("Got %d points."), Vertices.Num());
+
+		auto ivIt = intVectorVars.Find("pointref.indices");
+		if (ivIt != nullptr)
+		{
+			auto iIt = intVars.Find("Tetrahedron_run:startvertex");
+			int32 startIndex = iIt == nullptr ? 0 : *iIt;
+
+			iIt = intVars.Find("Tetrahedron_run:nprimitives");
+			int32 numTets = iIt == nullptr ? -1 : *iIt;
+
+			UE_LOG(LogChaosFlesh, Display, TEXT("Tet start index: %d num tets: %d"), startIndex, numTets);
+
+			const TArray<int32>& indices = *ivIt;
+			Elements.Reserve(indices.Num() / 4);
+			size_t stopIndex = numTets != -1 ? startIndex + numTets * 4 : indices.Num();
+			for (size_t i = startIndex; i < stopIndex; i += 4)
+			{
+				FIntVector4 tet;
+				for (size_t j = 0; j < 4; j++)
+					tet[j] = indices[i + j];
+				Elements.Add(tet);
+			}
+			UE_LOG(LogChaosFlesh, Display, TEXT("Got %d tets."), Elements.Num());
+		}
+
+		GetSurfaceElements(Elements, SurfaceElements, false);
+		UE_LOG(LogChaosFlesh, Display, TEXT("Got %d tris."), SurfaceElements.Num());
+
+		Collection.Reset(FFleshCollection::NewFleshCollection(Vertices, SurfaceElements, Elements));
+		return Collection;
+	}
+
+	TUniquePtr<FFleshCollection> ReadTetFile(const FString& Filename)
+	{
+		UE_LOG(LogChaosFlesh, Display, TEXT("Reading Path %s"), *Filename);
+		TUniquePtr<FFleshCollection> Collection;
+
+		TArray<FVector> Vertices;
+		TArray<FIntVector4> Elements;  //UE::Math::TIntVector4<int32>
+		TArray<FIntVector3> SurfaceElements;
+
+		TArray<UE::Math::TVector<float>> FloatPos;
+		TArray<Chaos::TVector<int32, 4>> Tets;
+		if (IO::ReadStructure<4>(Filename, FloatPos, Tets))
+		{
+			Vertices.SetNum(FloatPos.Num());
+			for (int32 i = 0; i < FloatPos.Num(); i++)
+				for (int32 j = 0; j < 3; j++)
+					Vertices[i][j] = FloatPos[i][j];
+
+			Elements.SetNum(Tets.Num());
+			for (int32 i = 0; i < Tets.Num(); i++)
+				for (int32 j = 0; j < 4; j++)
+					Elements[i][j] = Tets[i][j];
+
+			UE_LOG(LogChaosFlesh, Display, TEXT("Got %d tets."), Elements.Num());
+		}
+
+		GetSurfaceElements(Elements, SurfaceElements, false);
+		UE_LOG(LogChaosFlesh, Display, TEXT("Got %d tris."), SurfaceElements.Num());
+
+		Collection.Reset(FFleshCollection::NewFleshCollection(Vertices, SurfaceElements, Elements));
+		return Collection;
+	}
+
+	TUniquePtr<FFleshCollection> ImportTetFromFile(const FString& Filename)
+	{
+		if (FPaths::FileExists(Filename))
+		{
+			if (Filename.EndsWith(".tet") || Filename.EndsWith(".tet.gz"))
+				return ReadTetFile(Filename);
+			else if (Filename.EndsWith(".geo"))
+				return ReadGEOFile(Filename);
+			UE_LOG(LogChaosFlesh, Warning, TEXT("Unsupported file type: '%s'."), *Filename);
+		}
+		else
+		{
+			UE_LOG(LogChaosFlesh, Warning, TEXT("Unknown file path: '%s'."), *Filename);
+		}
+		return nullptr;
+	}
+}
diff --git a/Engine/Plugins/Experimental/ChaosFlesh/Source/ChaosFlesh/Private/ChaosFlesh/GEO.h b/Engine/Plugins/Experimental/ChaosFlesh/Source/ChaosFlesh/Private/ChaosFlesh/GEO.h
new file mode 100644
--- /dev/null
+++ b/Engine/Plugins/Experimental/ChaosFlesh/Source/ChaosFlesh/Private/ChaosFlesh/GEO.h
@@ -0,0 +1,20 @@
+// Copyright Epic Games, Inc. All Rights Reserved.
+#pragma once
+
+#include "Containers/Map.h"
+#include "Containers/UnrealString.h"
+
+#include <string>
+
+namespace ChaosFlesh
+{
+	bool
+	ReadGEO(
+		const std::string& filename,
+		TMap<FString, int32>& IntVars,
+		TMap<FString, TArray<int32>>& IntVectorVars,
+		TMap<FString, TArray<float>>& FloatVectorVars,
+		TMap<FString, TPair<TArray<std::string>, TArray<int32>>>& IndexedStringVars,
+		std::ostream* errorStream=nullptr);
+
+} // namespace ChaosFlesh
\ No newline at end of file
diff --git a/Engine/Plugins/Experimental/ChaosFlesh/Source/ChaosFlesh/Private/ChaosFlesh/PB.cpp b/Engine/Plugins/Experimental/ChaosFlesh/Source/ChaosFlesh/Private/ChaosFlesh/PB.cpp
new file mode 100644
--- /dev/null
+++ b/Engine/Plugins/Experimental/ChaosFlesh/Source/ChaosFlesh/Private/ChaosFlesh/PB.cpp
@@ -0,0 +1,877 @@
+// Copyright Epic Games, Inc. All Rights Reserved.
+#include "ChaosFlesh/PB.h"
+#include "ChaosFlesh/FleshCollectionUtility.h" // for LogChaosFlesh
+#include "ChaosFlesh/IFileStream.h"
+
+#ifdef USE_ZLIB
+#include "ChaosFlesh/ZIP.h"
+#endif
+
+#include "Containers/StringConv.h"
+#include "GenericPlatform/GenericPlatformFile.h"
+#include "GenericPlatform/GenericPlatformMath.h"
+#include "Misc/FileHelper.h"
+#include "Misc/Paths.h"
+
+#include "HAL/PlatformFileManager.h"
+
+#include <fstream>
+#include <iostream>
+#include <limits>
+#include <string>
+#include <sstream>
+
+//============================================================================
+// Helper functions
+//============================================================================
+
+std::shared_ptr<std::istream>
+SafeOpenInput(const FString& filename, const bool binary = true)
+{
+	FPlatformFileManager& FileManager = FPlatformFileManager::Get();
+	IPlatformFile& PlatformFile = FileManager.GetPlatformFile();
+	if (!PlatformFile.FileExists(*filename))
+	{
+		return nullptr;
+	}
+	const bool compressed = filename.EndsWith(".gz");
+	if (compressed || binary)
+	{
+#ifdef USE_ZLIB
+		IFileHandle* infile = PlatformFile.OpenRead(*filename, false);
+		return std::shared_ptr<std::istream>(new ZIP_FILE_ISTREAM(infile, false));
+#else
+		UE_LOG(LogChaosFlesh, Display, TEXT("ChaosFlesh not compiled with zlib support!"));
+		return nullptr;
+#endif
+	}
+	else
+	{
+		IFileHandle* infile = PlatformFile.OpenRead(*filename, false);
+		return std::shared_ptr<std::istream>(new IFileStream(infile));
+	}
+	return nullptr;
+}
+
+//============================================================================
+// Read
+//============================================================================
+
+template <class T>
+void
+Read(std::istream& InFile, T& x, T& y, T& z, T& w)
+{
+	std::string Input;
+	std::stringstream ss;
+	check(std::getline(InFile, Input, ','));
+	ss << Input << " ";
+	check(std::getline(InFile, Input, ','));
+	ss << Input << " ";
+	check(std::getline(InFile, Input, ','));
+	ss << Input << " ";
+	check(std::getline(InFile, Input, ','));
+	ss << Input << " ";
+	ss >> x;
+	ss >> y;
+	ss >> z;
+	ss >> w;
+}
+
+template <class T>
+void
+Read(std::istream& InFile, T& x, T& y, T& z)
+{
+	std::string Input;
+	std::stringstream ss;
+	check(std::getline(InFile, Input, ','));
+	ss << Input << " ";
+	check(std::getline(InFile, Input, ','));
+	ss << Input << " ";
+	check(std::getline(InFile, Input, ','));
+	ss << Input << " ";
+	ss >> x;
+	ss >> y;
+	ss >> z;
+}
+
+template <class T>
+void
+Read(std::istream& InFile, T& x, T& y)
+{
+	std::string Input;
+	std::stringstream ss;
+	check(std::getline(InFile, Input, ','));
+	ss << Input << " ";
+	check(std::getline(InFile, Input, ','));
+	ss << Input << " ";
+	ss >> x;
+	ss >> y;
+}
+
+template <class T>
+void
+Read(std::istream& InFile, T& x)
+{
+	std::string Input;
+	std::stringstream ss;
+	check(std::getline(InFile, Input, ','));
+	ss << Input << " ";
+	ss >> x;
+}
+
+//============================================================================
+// ReadBinary
+//============================================================================
+
+template <class T>
+void
+ReadBinary(std::istream& InFile, T& x)
+{
+	if (!InFile.eof())
+		InFile.read(reinterpret_cast<char*>(&x), sizeof(T));
+}
+
+template <class T, int d>
+void
+ReadBinary(std::istream& InFile, Chaos::TVector<T, d>& Val)
+{
+	check(sizeof(Chaos::TVector<T, d>) == sizeof(T) * d);
+	InFile.read(reinterpret_cast<char*>(&Val[0]), sizeof(T) * d);
+}
+
+template <class T, int d>
+void
+ReadBinary(std::istream& InFile, TArray<Chaos::TVector<T, d>>& Values)
+{
+	check(sizeof(Chaos::TVector<T, d>) == sizeof(T) * d);
+	InFile.read(reinterpret_cast<char*>(&Values[0]), sizeof(T) * d * Values.Num());
+}
+
+// LWC safe double FReal specializations - read as floats
+void
+ReadBinary(std::istream& InFile, UE::Math::TVector<double>& Val)
+{
+	Chaos::TVector<float, 3> AsFloat;
+	ReadBinary(InFile, AsFloat);
+	Val[0] = AsFloat[0];
+	Val[1] = AsFloat[1];
+	Val[2] = AsFloat[2];
+}
+
+void
+ReadBinary(std::istream& InFile, TArray<UE::Math::TVector<double>>& Values)
+{
+	// Per entry float -> double conversion
+	for (int32 i = 0; i < Values.Num(); i++)
+		ReadBinary(InFile, Values[i]);
+}
+
+
+template <class T>
+void
+ReadBinary(std::istream& InFile, TArray<T>& Values)
+{
+	InFile.read(reinterpret_cast<char*>(&Values[0]), sizeof(T) * Values.Num());
+}
+
+bool
+ReadBinary(std::istream& InFile, std::string& str)
+{
+	int32 Size = -1;
+	ReadBinary(InFile, Size);
+	if (Size > 0 && Size < 1024 * 1000)
+	{
+		str.resize(Size);
+		InFile.read(&str.front(), Size);
+		return true;
+	}
+	else
+	{
+		UE_LOG(LogChaosFlesh, Display, TEXT("ReadBinary: Invalid string length: %d"), Size);
+		return false;
+	}
+}
+
+bool
+ReadBinary(std::istream& InFile, FString& str)
+{
+	int32 Size = -1;
+	ReadBinary(InFile, Size);
+	if (Size > 0 && Size < 1024 * 1000)
+	{
+		str.Reset(Size);
+		InFile.read((char*) & str[0], Size);
+		//InFile.Read(reinterpret_cast<uint8*>(&str[0]), Size);
+		return true;
+	}
+	else
+	{
+		UE_LOG(LogChaosFlesh, Display, TEXT("ReadBinary: Invalid string length: %d"), Size);
+		return false;
+	}
+}
+
+
+namespace ChaosFlesh {
+namespace IO {
+
+template <int d>
+bool 
+ReadStructure(const FString& filename, TArray<UE::Math::TVector<float>>& positions, TArray<Chaos::TVector<int32,d>>& mesh)
+{
+	std::shared_ptr<std::istream> InFile = SafeOpenInput(filename);
+	if (!InFile)
+	{
+		UE_LOG(LogChaosFlesh, Display, TEXT("Failed to open file: '%s'"), *filename);
+		return false;
+	}
+	std::istream& InFileRef = *InFile;
+
+	//
+	// READ_WRITE_MESH_OBJECT implicitly forwards to READ_WRITE_SIMPLEX_MESH to read the mesh, then
+	// reads the number of particles in the same line.
+	//
+	int NumberNodes = -2;               // Should be non-zero
+	ReadBinary(InFileRef, NumberNodes);
+
+	int BackwardsCompatible = -2;       // simplex dimension + 1 (3 = tri, 4 = tet)
+	ReadBinary(InFileRef, BackwardsCompatible);
+	if (BackwardsCompatible != d)
+	{
+		UE_LOG(LogChaosFlesh, Display, TEXT("File: '%s' opened as type '%s' but has simplex dimension %d not %d."), 
+			*filename, (d == 4 ? "tet" : d == 3 ? "tri" : "curve"), BackwardsCompatible, d);
+	}
+
+	// READ_WRITE_ARRAY::Write_Prefix()
+	int ArraySize = -1;                 // called "prefix"
+	ReadBinary(InFileRef, ArraySize);
+	if (ArraySize < 0)
+	{
+		UE_LOG(LogChaosFlesh, Display, TEXT("File: '%s', invalid array size: %d."),
+			*filename, ArraySize);
+		return false;
+	}
+
+	mesh.SetNum(ArraySize);
+	for (int i = 0; i < ArraySize; i++)
+		for(int j=0; j < d; j++)
+		{
+			int Idx = -1;
+			ReadBinary(InFileRef, Idx);
+			// Convert from 1 indexed to 0 indexed.
+			mesh[i][j] = static_cast<int32>(Idx - 1);
+		}
+
+	int NumParticles = -1;
+	ReadBinary(InFileRef, NumParticles);
+
+	//return ReadPoints(InFile, filename, NumParticles, positions);
+	positions.SetNum(NumParticles);
+	for (int i = 0; i < NumParticles; i++)
+		ReadBinary(InFileRef, positions[i]);
+	return true;
+}
+
+template bool ReadStructure<2>(const FString&, TArray<UE::Math::TVector<float>>&, TArray<Chaos::TVector<int32, 2>>&);
+template bool ReadStructure<3>(const FString&, TArray<UE::Math::TVector<float>>&, TArray<Chaos::TVector<int32, 3>>&);
+template bool ReadStructure<4>(const FString&, TArray<UE::Math::TVector<float>>&, TArray<Chaos::TVector<int32, 4>>&);
+
+DeformableGeometryCollectionReader::DeformableGeometryCollectionReader(
+	const FString& BaseDir, 
+	const FString& ColorFilePath, 
+	const FString& ColorGeometryFilePath)
+	: BaseDir(BaseDir)
+	, CommonDir(BaseDir+"/common")
+	, ColorFilePath(ColorFilePath)
+	, ColorGeometryFilePath(ColorGeometryFilePath)
+{}
+
+DeformableGeometryCollectionReader::~DeformableGeometryCollectionReader()
+{}
+
+bool
+DeformableGeometryCollectionReader::ReadPBScene()
+{
+	//FString ColorFilePath = ColorMap.FilePath;
+	if (!FPaths::FileExists(ColorFilePath))
+	{
+		UE_LOG(LogChaosFlesh, Display, TEXT("PhysBAM color map file not found: '%s'"), *ColorFilePath);
+		// Don't fail.
+	}
+	else
+	{
+		// Read the color files to the PBVertexColors array, which if authored, will be used
+		// at procedural mesh construction time for vertex colors.
+		ReadColorFile(ColorFilePath);
+	}
+
+	//FString ColorGeometryFilePath = ColorMapGeometry.FilePath;
+	if (!FPaths::FileExists(ColorFilePath))
+	{
+		UE_LOG(LogChaosFlesh, Display, TEXT("PhysBAM color map geometry file not found: '%s'"), *ColorGeometryFilePath);
+		// Don't fail.
+	}
+	else
+	{
+		// Read the color files to the PBVertexColors array, which if authored, will be used
+		// at procedural mesh construction time for vertex colors.
+		ReadColorGeometryFile(ColorGeometryFilePath);
+	}
+
+	//BaseDir = BaseDirectory.Path;
+	if (!FPaths::DirectoryExists(BaseDir))
+	{
+		UE_LOG(LogChaosFlesh, Display, TEXT("PhysBAM simulation base directory not found: '%s'"), *BaseDir);
+		return false;
+	}
+	UE_LOG(LogChaosFlesh, Display, TEXT("PhysBAM simulation base directory: '%s'"), *BaseDir);
+	CommonDir = BaseDir + TEXT("/common");
+	if (!FPaths::DirectoryExists(CommonDir))
+	{
+		UE_LOG(LogChaosFlesh, Display, TEXT("PhysBAM simulation common directory not found: '%s'"), *BaseDir);
+		return false;
+	}
+
+	FString DefObjStructPath = CommonDir + TEXT("/deformable_object_structures");
+	if (!FPaths::FileExists(DefObjStructPath))
+	{
+		DefObjStructPath = CommonDir + TEXT("/deformable_object_structures.gz");
+		if (!FPaths::FileExists(DefObjStructPath))
+		{
+			UE_LOG(LogChaosFlesh, Display, TEXT("PhysBAM file not found: '%s'"), *DefObjStructPath);
+			return false;
+		}
+	}
+	if (!ReadDeformableGeometryCollection(DefObjStructPath))
+		return false;
+
+	const FString FaceControlPath = BaseDir + TEXT("/face_control_parameters_configuration.gz");
+	if (FPaths::FileExists(FaceControlPath))
+	{
+		if (!ReadFaceControlParameters(FaceControlPath))
+			return false;
+	}
+
+	return true;
+}
+
+TPair<int32, int32>
+DeformableGeometryCollectionReader::ReadFrameRange()
+{
+	int32 FFrame = -std::numeric_limits<int32>::max();
+	int32 LFrame = -std::numeric_limits<int32>::max();
+	ReadFrameRange(CommonDir, FFrame, LFrame);
+	return TPair<int32, int32>(FFrame, LFrame);
+}
+
+void
+DeformableGeometryCollectionReader::ReadFrameRange(
+	const FString& CommonDirIn,
+	int32& InputFFrameIn,
+	int32& InputLFrameIn) const
+{
+	if (InputFFrameIn == -std::numeric_limits<int32>::max())
+	{
+		const FString FFramePath = CommonDirIn + TEXT("/first_frame");
+		if (FPaths::FileExists(FFramePath))
+		{
+			if (std::shared_ptr<std::istream> FFrameFile = SafeOpenInput(FFramePath, false))
+			{
+				float Tmp = -std::numeric_limits<float>::max();
+				::Read<float>(*FFrameFile, Tmp); // Static analysis requires the fully qualified name
+				if (Tmp >= 0)
+				{
+					InputFFrameIn = Tmp;
+				}
+			}
+		}
+	}
+
+	// This file is frequently re-written as the sim progresses, so it's pretty common to get read errors!
+	const FString LFramePath = CommonDirIn + TEXT("/last_frame");
+	if (FPaths::FileExists(LFramePath))
+	{
+		if (std::shared_ptr<std::istream> LFrameFile = SafeOpenInput(LFramePath, false))
+		{
+			float Tmp = -std::numeric_limits<float>::max();
+			::Read<float>(*LFrameFile, Tmp);  // Static analysis requires the fully qualified name
+			if (Tmp >= InputFFrameIn)
+			{
+				InputLFrameIn = Tmp;
+			}
+		}
+	}
+}
+
+int32
+DeformableGeometryCollectionReader::ReadNumPoints(
+	const int32 Frame,
+	const FString& BaseDirIn) const
+{
+	const FString DefObjParticlesPath = FString::Printf(TEXT("%s/%d/deformable_object_particles.gz"), *BaseDirIn, Frame);
+	if (FPaths::FileExists(DefObjParticlesPath))
+	{
+		if (std::shared_ptr<std::istream> InFile = SafeOpenInput(DefObjParticlesPath))
+		{
+			// READ_WRITE_POINT_CLOUD::Read()
+			int32 Version = -1;
+			ReadBinary(*InFile, Version);
+			if (Version != 1)
+				return 0;
+
+			// READ_WRITE_ARRAY_COLLECTION::Read()
+			int32 Size = -1;
+			ReadBinary(*InFile, Size);
+			return Size;
+		}
+	}
+	return 0;
+}
+
+bool
+DeformableGeometryCollectionReader::ReadPoints(
+	const int32 Frame,
+	TArray<UE::Math::TVector<float>>& Points)
+{
+	const FString DefObjParticlesPath = FString::Printf(TEXT("%s/%d/deformable_object_particles.gz"), *BaseDir, Frame);
+	if (FPaths::FileExists(DefObjParticlesPath))
+	{
+		if (std::shared_ptr<std::istream> InFile = SafeOpenInput(DefObjParticlesPath))
+		{
+			// READ_WRITE_POINT_CLOUD::Read()
+			int32 Version = -1;
+			ReadBinary(*InFile, Version);
+			if (Version != 1)
+			{
+				UE_LOG(LogChaosFlesh, Display, TEXT("File: '%s' - Unrecognized particle version: %d"), *DefObjParticlesPath, Version);
+				return false;
+			}
+
+			// READ_WRITE_ARRAY_COLLECTION::Read()
+			int32 Size = -1;
+			ReadBinary(*InFile, Size);
+			if (Size <= 0)
+			{
+				UE_LOG(LogChaosFlesh, Display, TEXT("File: '%s' - Invalid array collection size: %d"), *DefObjParticlesPath, Size);
+				return false;
+			}
+
+			// READ_WRITE_ARRAY_COLLECTION::Read_Arrays()
+			int32 SecondSize = -1;
+			ReadBinary(*InFile, SecondSize);
+			int32 NumAttributes = -1;
+			ReadBinary(*InFile, NumAttributes);
+			if (NumAttributes <= 0 || NumAttributes > 100)
+			{
+				UE_LOG(LogChaosFlesh, Display, TEXT("File: '%s' - Invalid num attributes: %d"), *DefObjParticlesPath, NumAttributes);
+				return false;
+			}
+
+			Points.SetNumUninitialized(Size);
+			for (auto& Pt : Points)
+				for (int32 i = 0; i < 3; i++)
+				{
+					Pt[i] = -std::numeric_limits<float>::max();
+				}
+
+			for (int32 i = 1; i <= NumAttributes; i++)
+			{
+				int32 AttributeId;
+				ReadBinary(*InFile, AttributeId);
+				AttributeId &= 0x0000FFFF;		// Type_Only(), strips lower bits
+
+				// GEOMETRY_PARTICLES has X, ID, and maybe V
+
+				if (AttributeId == 1) // ATTRIBUTE_ID_X
+				{
+					// READ_WRITE_ARRAY::Read()
+					int32 ArraySize = -1;
+					ReadBinary(*InFile, ArraySize);
+
+					int32 ArraySize2 = -1;
+					ReadBinary(*InFile, ArraySize2);
+
+					//UE_LOG(LogChaosFlesh, Display, TEXT("Ryan - Size: %d, ArraySize: %d, ArraySize2: %d"),
+					//	Size, ArraySize, ArraySize2);
+
+					ReadBinary(*InFile, Points);
+
+					// All we care about is X.  Return.
+					return true;
+				}
+				else if (AttributeId == 2) // ATTRIBUTE_ID_V
+				{
+					// READ_WRITE_ARRAY::Read()
+					int32 ArraySize = -1;
+					ReadBinary(*InFile, ArraySize);
+					if (ArraySize != Size)
+					{
+					}
+					// We don't care about V, but it came first.  So just read it, 
+					// and we'll write over it.
+					ReadBinary(*InFile, Points);
+				}
+				else if (AttributeId == 20 || // ATTRIBUTE_ID_ID
+					AttributeId == 6)		 // ATTRIBUTE_ID_STRUCTURE_IDS
+				{
+					// READ_WRITE_ARRAY::Read()
+					int32 ArraySize = -1;
+					ReadBinary(*InFile, ArraySize);
+					if (ArraySize != Size)
+					{
+					}
+
+					// Just read ID and then throw it away.
+					TArray<int32> Id;
+					Id.SetNumUninitialized(Size);
+					ReadBinary(*InFile, Id);
+				}
+				else
+				{
+					UE_LOG(LogChaosFlesh, Display, TEXT("File: '%s' - Unsupported attribute id: %d"), *DefObjParticlesPath, AttributeId);
+					return false;
+				}
+			}
+		}
+		else
+		{
+			UE_LOG(LogChaosFlesh, Display, TEXT("Failed to open: '%s'"), *DefObjParticlesPath);
+			return false;
+		}
+	}
+	else
+	{
+		UE_LOG(LogChaosFlesh, Display, TEXT("File not found: '%s'"), *DefObjParticlesPath);
+		return false;
+	}
+	return false;
+}
+
+bool
+DeformableGeometryCollectionReader::ReadDeformableGeometryCollection(const FString& FilePath)
+{
+	std::shared_ptr<std::istream> InFile = SafeOpenInput(FilePath);
+	if (!InFile)
+	{
+		UE_LOG(LogChaosFlesh, Display, TEXT("Failed to open file: '%s'"), *FilePath);
+		return false;
+	}
+
+	// DEFORMABLE_GOMETRY_COLLECTION::Read_Static_Variables() uses TYPED_ISTREAM to switch between reading float and double.
+	int32 NumStructures = -1;
+	ReadBinary(*InFile, NumStructures);
+	if (NumStructures <= 0)
+	{
+		UE_LOG(LogChaosFlesh, Display, TEXT("File: '%s' - Invalid num structures: %d"), *FilePath, NumStructures);
+		return false;
+	}
+	UE_LOG(LogChaosFlesh, Display, TEXT("File: '%s' - Reading num structures: %d"), *FilePath, NumStructures);
+
+	for (int SIdx = 1; SIdx <= NumStructures; ++SIdx)
+	{
+		// READ_WRITE_STRUCTURE::Create_Structure()
+		std::string TypeName;
+		if (!ReadBinary(*InFile, TypeName))
+		{
+			return InputMeshes.Num() > 0;
+		}
+
+		// STRUCTURE::Create_From_Name() -> STRUCTURE_REGISTRY::Name_To_Factory()
+		// STRUCTURE_REGISTRY is keyed by T_OBJECT::Static_Extension()
+		//		tri = TRIANGUALTED_SURFACE
+		//		tet = TETRAHEDRALIZED_VOLUME = SIMPLICIAL_OBJECT<T,VECTOR<T,3>,3>
+		//		curve = SEGMENTED_CURVE
+		// Geometry read routines forward to READ_WRITE_MESH_OBJECT.h
+
+		if (TypeName == "SIMPLICIAL_OBJECT<T,VECTOR<T,3>,3>")
+			TypeName = "tet";
+		else if (TypeName == "SIMPLICIAL_OBJECT<T,VECTOR<T,3>,2>")
+			TypeName = "tri";
+		else if (TypeName == "SIMPLICIAL_OBJECT<T,VECTOR<T,3>,1>")
+			TypeName = "curve";
+
+		if (TypeName != "tet" && TypeName != "tri")
+		{
+			FString TmpStr(TypeName.c_str());
+			UE_LOG(LogChaosFlesh, Display, TEXT("Unsupported PhysBAM structure type: '%s'"), *TmpStr);
+		}
+		else
+		{
+			//
+			// READ_WRITE_MESH_OBJECT implicitly forwards to READ_WRITE_SIMPLEX_MESH to read the mesh, then
+			// reads the number of particles in the same line.
+			//
+			int NumberNodes = -1;				// Should be non-zero
+			ReadBinary(*InFile, NumberNodes);
+
+			int BackwardsCompatible = -1;		// simplex dimension + 1 (3 = tri, 4 = tet)
+			ReadBinary(*InFile, BackwardsCompatible);
+
+			// READ_WRITE_ARRAY::Write_Prefix()
+			int ArraySize = -1;					// called "prefix" 
+			ReadBinary(*InFile, ArraySize);
+			if (ArraySize < 0)
+			{
+				UE_LOG(LogChaosFlesh, Display, TEXT("File: '%s', structure: %d of %d - Invalid array size: %d"),
+					*FilePath, SIdx, NumStructures, ArraySize);
+				return InputMeshes.Num() > 0;
+			}
+
+			Mesh* mesh = nullptr;
+			int32 MinIdx = std::numeric_limits<int32>::max();
+			int32 MaxIdx = -std::numeric_limits<int32>::max();
+			int32 NumTriangles = 0;
+			if (TypeName == "tet")
+			{
+				TetMesh* tetMesh = new TetMesh();
+				mesh = static_cast<Mesh*>(tetMesh);
+				InputMeshes.Push(mesh);
+				TetMeshes.Push(tetMesh);
+
+				tetMesh->Elements.SetNum(ArraySize);
+				for (auto& Elem : tetMesh->Elements)
+					for (int i = 0; i < 4; i++)
+						Elem[i] = -std::numeric_limits<int32>::max();
+				ReadBinary(*InFile, tetMesh->Elements);
+
+				size_t BytesToEOF = 0;
+				while (!InFile->eof() && !InFile->fail())
+				{
+					char ch;
+					InFile->read(&ch, 1);
+					BytesToEOF++;
+				}
+				UE_LOG(LogChaosFlesh, Display, TEXT("Done reading: '%s' (%d bytes remaining)"),
+					*FilePath, BytesToEOF);
+
+				// Convert Elements from 1 indexed to 0 indexed
+				int32 ElemIdx = -1;
+				for (auto& Elem : tetMesh->Elements)
+				{
+					ElemIdx++;
+					for (int i = 0; i < 4; i++)
+					{
+						if (Elem[i] == -std::numeric_limits<int32>::max() || Elem[i] <= 0 || Elem[i] > NumberNodes)
+						{
+							UE_LOG(LogChaosFlesh, Display, TEXT("File: '%s', structure: %d of %d - File read error - elem %d of %d, index %d, coordinate value: %d"),
+								*FilePath, ElemIdx, tetMesh->Elements.Num(), i, Elem[i]);
+							continue;
+						}
+						Elem[i]--;
+						check(Elem[i] >= 0);
+					}
+				}
+
+				UE_LOG(LogChaosFlesh, Display, TEXT("File: '%s', structure: %d of %d - Num elements: %d"),
+					*FilePath, SIdx, NumStructures, ArraySize);
+
+				for (auto& Elem : tetMesh->Elements)
+				{
+					for (int i = 0; i < 4; i++)
+					{
+						const int32 Idx = Elem[i];
+						MinIdx = Idx < MinIdx ? Idx : MinIdx;
+						MaxIdx = Idx > MaxIdx ? Idx : MaxIdx;
+					}
+				}
+			}
+			else //if(TypeName == "tri")
+			{
+				TriMesh* triMesh = new TriMesh();
+				mesh = static_cast<Mesh*>(triMesh);
+				InputMeshes.Push(mesh);
+				TriMeshes.Push(triMesh);
+
+				triMesh->SurfaceElements.SetNum(ArraySize);
+				ReadBinary(*InFile, triMesh->SurfaceElements);
+
+				// Convert Elements from 1 indexed to 0 indexed, and reverse orientation
+				for (auto& Elem : triMesh->SurfaceElements)
+				{
+					std::swap(Elem[0], Elem[1]);
+					for (int i = 0; i < 3; i++)
+						Elem[i]--;
+				}
+
+				UE_LOG(LogChaosFlesh, Display, TEXT("File: '%s', structure: %d of %d - Num elements: %d"),
+					*FilePath, SIdx, NumStructures, ArraySize);
+
+				NumTriangles = triMesh->SurfaceElements.Num();
+				for (auto& Elem : triMesh->SurfaceElements)
+				{
+					for (int i = 0; i < 3; i++)
+					{
+						const int32 Idx = Elem[i];
+						MinIdx = Idx < MinIdx ? Idx : MinIdx;
+						MaxIdx = Idx > MaxIdx ? Idx : MaxIdx;
+					}
+				}
+			}
+
+			//
+			// Back to READ_WRITE_MESH_OBJECT::Read_Helper(), read points
+			//
+
+			int NumPoints = 0;
+			ReadBinary(*InFile, NumPoints);
+			if (NumPoints > 0)
+			{
+				mesh->Points.SetNum(NumPoints);
+				ReadBinary(*InFile, mesh->Points);
+			}
+
+			UE_LOG(LogChaosFlesh, Display,
+				TEXT("Read file '%s' structure %d of %d:\n"
+					"    type: %s\n"
+					"    num elements: %d\n"
+					"    num triangles: %d\n"
+					"    min/max point index: %d, %d\n"
+					"    num nodes: %d\n"
+					"    num points: %d"),
+				*FilePath, SIdx, NumStructures,
+				*FString(TypeName.c_str()),
+				ArraySize,
+				NumTriangles,
+				MinIdx, MaxIdx,
+				NumberNodes,
+				NumPoints);
+		}
+	}
+	return true;
+}
+
+bool
+DeformableGeometryCollectionReader::ReadFaceControlParameters(const FString& FilePath)
+{
+	// See FACE_CONTROL_PARAMETERS::Read_Configuration_From_File()
+	return true;
+}
+
+bool
+DeformableGeometryCollectionReader::ReadColorFile(const FString& FilePath)
+{
+	std::shared_ptr<std::istream> InFile = SafeOpenInput(FilePath);
+	if (!InFile)
+	{
+		UE_LOG(LogChaosFlesh, Display, TEXT("Failed to open file: '%s'"), *FilePath);
+		return false;
+	}
+
+	int32 Size = 0;
+	ReadBinary(*InFile, Size);
+
+	PBVertexColors.AddUninitialized(Size);
+	ReadBinary(*InFile, PBVertexColors);
+
+	Chaos::TVector<float, 3> MinPt(std::numeric_limits<float>::max(), std::numeric_limits<float>::max(), std::numeric_limits<float>::max());
+	Chaos::TVector<float, 3> MaxPt(-std::numeric_limits<float>::max(), -std::numeric_limits<float>::max(), -std::numeric_limits<float>::max());
+	for (const auto& Pt : PBVertexColors)
+	{
+		for (int j = 0; j < 3; j++)
+		{
+			if (Pt[j] < MinPt[j])
+				MinPt[j] = Pt[j];
+			if (Pt[j] > MaxPt[j])
+				MaxPt[j] = Pt[j];
+		}
+	}
+
+	UE_LOG(LogChaosFlesh, Display,
+		TEXT("File: '%s' - Read %d vertex colors:\n"
+			"    R range: %.6f, %.6f\n"
+			"    G range: %.6f, %.6f\n"
+			"    B range: %.6f, %.6f"),
+		*FilePath, Size,
+		MinPt[0], MaxPt[0],
+		MinPt[1], MaxPt[1],
+		MinPt[2], MaxPt[2]);
+	return true;
+}
+
+bool
+DeformableGeometryCollectionReader::ReadColorGeometryFile(const FString& FilePath)
+{
+	std::shared_ptr<std::istream> InFile = SafeOpenInput(FilePath);
+	if (!InFile)
+	{
+		UE_LOG(LogChaosFlesh, Display, TEXT("Failed to open file: '%s'"), *FilePath);
+		return false;
+	}
+
+	int32 NumNodes;
+	ReadBinary(*InFile, NumNodes);
+
+	int32 BackwardsCompatible; // simplex dimension + 1 (3 = tri, 4 = tet)
+	ReadBinary(*InFile, BackwardsCompatible);
+
+	if (BackwardsCompatible != 3)
+	{
+		UE_LOG(LogChaosFlesh, Display, TEXT("File: '%s', expected type 3 (tri), got %d"), *FilePath, BackwardsCompatible);
+		return false;
+	}
+
+	int32 NumElements;
+	ReadBinary(*InFile, NumElements);
+
+	PBVertexColorsTriSurf.SurfaceElements.SetNumUninitialized(NumElements);
+	ReadBinary(*InFile, PBVertexColorsTriSurf.SurfaceElements);
+
+	int32 NumParticles;
+	ReadBinary(*InFile, NumParticles);
+
+	PBVertexColorsTriSurf.Points.SetNumUninitialized(NumParticles);
+	ReadBinary(*InFile, PBVertexColorsTriSurf.Points);
+
+	// Convert Elements from 1 indexed to 0 indexed
+	for (auto& Elem : PBVertexColorsTriSurf.SurfaceElements)
+	{
+		for (int i = 0; i < 3; i++)
+			Elem[i]--;
+	}
+
+	int32 MinIdx = std::numeric_limits<int32>::max();
+	int32 MaxIdx = -std::numeric_limits<int32>::max();
+	for (auto& Elem : PBVertexColorsTriSurf.SurfaceElements)
+	{
+		for (int i = 0; i < 3; i++)
+		{
+			const int32 Idx = Elem[i];
+			MinIdx = Idx < MinIdx ? Idx : MinIdx;
+			MaxIdx = Idx > MaxIdx ? Idx : MaxIdx;
+		}
+	}
+
+	Chaos::TVector<float,3> MinPt(std::numeric_limits<float>::max(), std::numeric_limits<float>::max(), std::numeric_limits<float>::max());
+	Chaos::TVector<float,3> MaxPt(-std::numeric_limits<float>::max(), -std::numeric_limits<float>::max(), -std::numeric_limits<float>::max());
+	for (const auto& Pt : PBVertexColorsTriSurf.Points)
+	{
+		for (int j = 0; j < 3; j++)
+		{
+			if (Pt[j] < MinPt[j])
+				MinPt[j] = Pt[j];
+			if (Pt[j] > MaxPt[j])
+				MaxPt[j] = Pt[j];
+		}
+	}
+
+	UE_LOG(LogChaosFlesh, Display,
+		TEXT("Read color geometry file: '%s':\n"
+			"    num triangles: %d\n"
+			"    min/max point index: %d, %d\n"
+			"    num nodes: %d\n"
+			"    num points: %d\n"
+			"    domain: (%.2f, %.2f, %.2f)x(%.2f, %.2f, %.2f)"),
+		*FilePath,
+		NumElements,
+		MinIdx, MaxIdx,
+		NumNodes,
+		NumParticles,
+		MinPt[0], MinPt[1], MinPt[2],
+		MaxPt[0], MaxPt[1], MaxPt[2]);
+	return true;
+}
+
+
+
+} // namespace IO
+} // namespace ChaosFlesh
diff --git a/Engine/Plugins/Experimental/ChaosFlesh/Source/ChaosFlesh/Private/ChaosFlesh/PB.h b/Engine/Plugins/Experimental/ChaosFlesh/Source/ChaosFlesh/Private/ChaosFlesh/PB.h
new file mode 100644
--- /dev/null
+++ b/Engine/Plugins/Experimental/ChaosFlesh/Source/ChaosFlesh/Private/ChaosFlesh/PB.h
@@ -0,0 +1,167 @@
+// Copyright Epic Games, Inc. All Rights Reserved.
+#pragma once
+
+#include "Chaos/Vector.h"
+#include "Containers/Array.h"
+
+//DECLARE_LOG_CATEGORY_EXTERN(LogChaosFlesh, Verbose, All);
+
+namespace ChaosFlesh
+{
+	namespace IO
+	{
+		/**
+		* Read PhysBAM geometry structure - tet (d=4), tri (d=3), curve(d=2).
+		*/
+		template <int d>
+		bool
+		ReadStructure(const FString& filename, TArray<UE::Math::TVector<float>>& positions, TArray<Chaos::TVector<int32, d>>& mesh);
+
+		inline bool ReadTet(const FString& filename, TArray<UE::Math::TVector<float>>& positions, TArray<Chaos::TVector<int32,4>>& mesh)
+		{ return ReadStructure<4>(filename, positions, mesh); }
+		inline bool ReadTri(const FString& filename, TArray<UE::Math::TVector<float>>& positions, TArray<Chaos::TVector<int32,3>>& mesh)
+		{ return ReadStructure<3>(filename, positions, mesh); }
+		inline bool ReadCurve(const FString& filename, TArray<UE::Math::TVector<float>>& positions, TArray<Chaos::TVector<int32,2>>& mesh)
+		{ return ReadStructure<2>(filename, positions, mesh); }
+
+		/**
+		 * PhysBAM simulation directory structure:
+		 *
+		 * <BaseDir>/common
+		 *		deformable_object_structures.gz					-> DEFORMABLE_GEOMETRY_COLLECTION
+		 *		first_frame
+		 *		last_frame
+		 *		log.txt
+		 *		muscle_densities.gz
+		 *		muscle_fibers.gz
+		 *		muscle_tetrahedra.gz
+		 *		rigid_body_names
+		 *		rigid_body_structures_key.gz
+		 *		sim.param
+		 * <BaseDir>/<FrameNum>
+		 *		control_parameters.gz
+		 *		deformable_object_particles.gz
+		 *		frame_title
+		 *		rigid_body_structure_active_ids.gz
+		 *		rigid_geometry_particles.gz
+		 *		time.gz
+		 * <BaseDir>/face_control_parameters_configuration.gz	-> FACE_CONTROL_PARAMETERS
+		 */
+		class DeformableGeometryCollectionReader
+		{
+		public:
+			struct Mesh
+			{
+				Mesh(const int Dimension)
+					: Dimension(Dimension)
+				{}
+				const int32 Dimension;
+				// Contains static points, not animated/simulated.
+				//TArray<Chaos::FVec3> Points;
+				//TArray<UE::Math::TVector<FReal>> Points; // double
+				TArray<FVector> Points; // UE::Math::TVector<double>
+			};
+
+			struct TriMesh : public Mesh
+			{
+				TriMesh(const int Dimension = 3)
+					: Mesh(3)
+				{}
+				TArray<Chaos::TVector<int32, 3>> SurfaceElements;
+			};
+
+			struct TetMesh : public Mesh
+			{
+				TetMesh()
+					: Mesh(4)
+				{}
+				TArray<Chaos::TVector<int32, 4>> Elements;
+			};
+
+		public:
+			DeformableGeometryCollectionReader(const FString& BaseDir, const FString& ColorFilePath=FString(), const FString& ColorGeometryFilePath=FString());
+			~DeformableGeometryCollectionReader();
+
+			bool
+			ReadPBScene();
+
+			TArray<Mesh*>
+			GetMeshes() const
+			{ return InputMeshes; }
+
+			TArray<TetMesh*>
+			GetTetMeshes() const
+			{ return TetMeshes; }
+
+			TArray<TriMesh*>
+			GetTriMeshes() const
+			{ return TriMeshes; }
+
+			const TArray<Chaos::TVector<float, 3>> &
+			GetVertexColors() const
+			{ return PBVertexColors; }
+
+			const TriMesh&
+			GetVertexColorsTriSurf() const
+			{ return PBVertexColorsTriSurf; }
+
+			//! Reads the currently available frane range.
+			TPair<int32, int32>
+			ReadFrameRange();
+
+			bool
+			ReadPoints(const int32 Frame, TArray<UE::Math::TVector<float>>& Points);
+
+			bool
+			ReadPoints(const int32 Frame, TArray<UE::Math::TVector<double>>& Points)
+			{
+				TArray<UE::Math::TVector<float>> FPoints;
+				if (!ReadPoints(Frame, FPoints))
+					return false;
+				Points.SetNum(FPoints.Num());
+				for (int32 i = 0; i < FPoints.Num(); i++)
+					for (int32 j = 0; j < 3; j++)
+						Points[i][j] = FPoints[i][j];
+				return true;
+			}
+
+		public:
+			//! Reads the available frame range.
+			//! \p CommonDir the 'BaseDir/common' directory to read from.
+			//! \p InputFFrame the first available frame; only read if value is != -std::numeric_limits<int32>::max().
+			//! \p InputLFrame the last available frame.
+			void
+			ReadFrameRange(const FString& CommonDir, int32& InputFFrame, int32& InputLFrame) const;
+
+			int32
+			ReadNumPoints(const int32 Frame, const FString& BaseDirIn) const;
+
+			bool
+			ReadDeformableGeometryCollection(const FString& FilePath);
+
+			bool
+			ReadFaceControlParameters(const FString& FilePath);
+
+			bool
+			ReadColorFile(const FString& FilePath);
+
+			bool
+			ReadColorGeometryFile(const FString& FilePath);
+
+		private:
+			FString BaseDir;
+			FString CommonDir;
+			FString ColorFilePath;
+			FString ColorGeometryFilePath;
+
+			TArray<Mesh*> InputMeshes;
+			TArray<TriMesh*> TriMeshes;
+			TArray<TetMesh*> TetMeshes;
+
+			TArray<Chaos::TVector<float, 3>> PBVertexColors;
+			TriMesh PBVertexColorsTriSurf;
+		};
+
+
+	}
+} // namespace ChaosFlesh
diff --git a/Engine/Plugins/Experimental/ChaosFlesh/Source/ChaosFlesh/Private/ChaosFlesh/ZIP.h b/Engine/Plugins/Experimental/ChaosFlesh/Source/ChaosFlesh/Private/ChaosFlesh/ZIP.h
new file mode 100644
--- /dev/null
+++ b/Engine/Plugins/Experimental/ChaosFlesh/Source/ChaosFlesh/Private/ChaosFlesh/ZIP.h
@@ -0,0 +1,629 @@
+// Copyright Epic Games, Inc. All Rights Reserved.
+#pragma once
+#ifdef USE_ZLIB
+#ifdef WITH_EDITOR
+
+#define USE_UE
+#ifdef USE_UE
+#include "Containers/Array.h"
+#include "Containers/Map.h"
+#include "Containers/UnrealString.h"
+#include "GenericPlatform/GenericPlatformFile.h"
+#include "GenericPlatform/GenericPlatformMath.h"
+#include "HAL/PlatformFileManager.h"
+#else
+#include <algorithm>
+#include <cassert>
+#include <fstream>
+#include <iomanip>
+#include <iostream>
+#include <memory>
+#include <stdexcept>
+#include <string>
+#endif
+
+#include <istream>
+#include <streambuf>
+
+THIRD_PARTY_INCLUDES_START
+#include "zlib.h"
+THIRD_PARTY_INCLUDES_END
+
+struct ZIP_FILE_HEADER
+{
+	unsigned short version;
+	unsigned short flags;
+	unsigned short compression_type;
+	unsigned short stamp_date, stamp_time;
+	unsigned int crc;
+	unsigned int compressed_size, uncompressed_size;
+#ifdef USE_UE
+	FString filename;
+#else
+	std::string filename;
+#endif
+	unsigned int header_offset; // local header offset
+
+	ZIP_FILE_HEADER()
+	{}
+
+#ifdef USE_UE
+	ZIP_FILE_HEADER(const FString& filename_input)
+#else
+	ZIP_FILE_HEADER(const std::string& filename_input)
+#endif
+		:version(20), flags(0), compression_type(8), stamp_date(0), stamp_time(0), crc(0),
+		compressed_size(0), uncompressed_size(0), filename(filename_input), header_offset(0)
+	{}
+
+#ifdef USE_UE
+	bool Read(IFileHandle* istream, const bool global, FString& err_msg)
+#else
+	bool Read(std::istream& istream, const bool global, std::string& err_msg)
+#endif
+	{
+		unsigned int sig;
+		// read and check for local/global magic
+		if (global) {
+			Read_Primitive(istream, sig);
+			if (sig != 0x02014b50) 
+			{ 
+				//std::cerr << "Did not find global header signature" << std::endl; 
+				err_msg += "ZIP_FILE_HEADER: Did not find global header signature.\n";
+				return false; 
+			}
+			Read_Primitive(istream, version);
+		}
+		else 
+		{
+			Read_Primitive(istream, sig);
+			if (sig != 0x04034b50) 
+			{ 
+				//LOG::cerr << "Did not find local header signature" << std::endl; 
+				err_msg += "ZIP_FILE_HEADER: Did not find local header signature.\n";
+				return false; 
+			}
+		}
+		// Read rest of header
+		Read_Primitive(istream, version);
+		Read_Primitive(istream, flags);
+		Read_Primitive(istream, compression_type);
+		Read_Primitive(istream, stamp_date);
+		Read_Primitive(istream, stamp_time);
+		Read_Primitive(istream, crc);
+		Read_Primitive(istream, compressed_size);
+		Read_Primitive(istream, uncompressed_size);
+		unsigned short filename_length, extra_length;
+		Read_Primitive(istream, filename_length);
+		Read_Primitive(istream, extra_length);
+		unsigned short comment_length = 0;
+		if (global) {
+			Read_Primitive(istream, comment_length); // filecomment
+			unsigned short disk_number_start, int_file_attrib;
+			unsigned int ext_file_attrib;
+			Read_Primitive(istream, disk_number_start); // disk# start
+			Read_Primitive(istream, int_file_attrib); // internal file
+			Read_Primitive(istream, ext_file_attrib); // ext final
+			Read_Primitive(istream, header_offset);
+		} // rel offset
+		char* buf = new char[std::max(comment_length, std::max(filename_length, extra_length))];
+#ifdef USE_UE
+		istream->Read(reinterpret_cast<uint8*>(buf), filename_length);
+#else
+		istream.read(buf, filename_length);
+#endif
+		buf[filename_length] = 0;
+#ifdef USE_UE
+		filename = FString(buf);
+		istream->Read(reinterpret_cast<uint8*>(buf), extra_length);
+		if (global) istream->Read(reinterpret_cast<uint8*>(buf), comment_length);
+#else
+		filename = std::string(buf);
+		istream.read(buf, extra_length);
+		if (global) istream.read(buf, comment_length);
+#endif
+		delete[] buf;
+		return true;
+	}
+
+protected:
+#ifdef USE_UE
+	template<class T>
+	void
+	Read_Primitive(IFileHandle* input, T& d)
+	{
+		//STATIC_ASSERT((sizeof(T) == PLATFORM_INDEPENDENT_SIZE<T>::value));
+		input->Read((uint8*)&d, sizeof(T));
+		//if (big_endian) Swap_Endianity(d); // convert to little endian if necessary
+	}
+#else
+	template<class T>
+	void
+	Read_Primitive(std::istream& input, T& d)
+	{
+		//STATIC_ASSERT((sizeof(T) == PLATFORM_INDEPENDENT_SIZE<T>::value));
+		input.read((char*)&d, sizeof(T));
+		//if (big_endian) Swap_Endianity(d); // convert to little endian if necessary
+	}
+#endif
+};
+
+struct GZIP_FILE_HEADER
+{
+	unsigned char magic0, magic1; // magic should be 0x8b,0x1f
+	unsigned char cm; // compression method 0x8 is gzip
+	unsigned char flags; // flags
+	unsigned int modtime; // 4 byte modification time
+	unsigned char flags2; // secondary flags
+	unsigned char os; // operating system 0xff for unknown
+	unsigned short crc16; // crc check
+	unsigned int crc32;
+
+	GZIP_FILE_HEADER()
+		:magic0(0), magic1(0), flags(0), modtime(0), flags2(0), os(0), crc16(0), crc32(0)
+	{}
+
+#ifdef USE_UE
+	bool Read(IFileHandle* istream, FString& err_msg)
+#else
+	bool Read(std::istream& istream, std::string& err_msg)
+#endif
+	{
+		Read_Primitive(istream, magic0);
+		Read_Primitive(istream, magic1);
+		if (magic0 != 0x1f || magic1 != 0x8b) 
+		{ 
+			//LOG::cerr << "gzip: did not find gzip magic 0x1f 0x8b" << std::endl;
+			err_msg += "gzip: did not find gzip magic 0x1f 0x8b";
+			return false; 
+		}
+		Read_Primitive(istream, cm);
+		if (cm != 8) 
+		{ 
+			//LOG::cerr << "gzip: compression method not 0x8" << std::endl;
+			err_msg += "gzip: compression method not 0x8";
+			return false; 
+		}
+		Read_Primitive(istream, flags);
+		Read_Primitive(istream, modtime);
+		Read_Primitive(istream, flags2);
+		Read_Primitive(istream, os);
+		unsigned char dummyByte;
+		// read flags if necessary
+		if (flags & 2) {
+			unsigned short flgExtraLen;
+			Read_Primitive(istream, flgExtraLen);
+			for (int k = 0; k < flgExtraLen; k++) 
+				Read_Primitive(istream, dummyByte);
+		}
+		// read filename/comment if present
+		int stringsToRead = ((flags & 8) ? 1 : 0) + ((flags & 4) ? 1 : 0);
+		for (int i = 0; i < stringsToRead; i++)
+			do { Read_Primitive(istream, dummyByte); } while (dummyByte != 0 && istream);
+		if (flags & 1) Read_Primitive(istream, crc16);
+#ifdef USE_UE
+		if(!istream || istream->Tell() >= istream->Size())
+#else
+		if (!istream) 
+#endif
+		{
+			//LOG::cerr << "gzip: got to end of file after only reading gzip header" << std::endl;
+			err_msg += "gzip: got to end of file after only reading gzip header";
+			return false; 
+		}
+		return true;
+	}
+protected:
+#ifdef USE_UE
+	template<class T>
+	void
+	Read_Primitive(IFileHandle* input, T& d)
+	{
+		//STATIC_ASSERT((sizeof(T) == PLATFORM_INDEPENDENT_SIZE<T>::value));
+		input->Read((uint8*)&d, sizeof(T));
+		//if (big_endian) Swap_Endianity(d); // convert to little endian if necessary
+	}
+#else
+	template<class T>
+	void
+	Read_Primitive(std::istream& input, T& d)
+	{
+		//STATIC_ASSERT((sizeof(T) == PLATFORM_INDEPENDENT_SIZE<T>::value));
+		input.read((char*)&d, sizeof(T));
+		//if (big_endian) Swap_Endianity(d); // convert to little endian if necessary
+}
+#endif
+};
+
+class ZIP_STREAMBUF_DECOMPRESS : public std::streambuf
+{
+	static const unsigned int buffer_size = 512;
+#ifdef USE_UE
+	IFileHandle* istream;
+#else
+	std::istream& istream;
+#endif
+
+	z_stream strm;
+	unsigned char in[buffer_size], out[buffer_size];
+	ZIP_FILE_HEADER header;
+	GZIP_FILE_HEADER gzip_header;
+	int total_read, total_uncompressed;
+	bool part_of_zip_file;
+	bool valid;
+	bool compressed_data;
+
+	static const unsigned short DEFLATE = 8;
+	static const unsigned short UNCOMPRESSED = 0;
+public:
+#ifdef USE_UE
+	ZIP_STREAMBUF_DECOMPRESS(IFileHandle* stream, bool part_of_zip_file_input)
+#else
+	ZIP_STREAMBUF_DECOMPRESS(std::istream& stream, bool part_of_zip_file_input)
+#endif
+		:istream(stream), total_read(0), total_uncompressed(0), part_of_zip_file(part_of_zip_file_input), valid(true)
+	{
+		strm.zalloc = Z_NULL; strm.zfree = Z_NULL; strm.opaque = Z_NULL; strm.avail_in = 0; strm.next_in = Z_NULL;
+		setg((char*)in, (char*)in, (char*)in);
+		setp(0, 0);
+		// skip the header
+#ifdef USE_UE
+		FString err_msg;
+#else
+		std::string err_msg;
+#endif
+		if (part_of_zip_file) {
+			valid = header.Read(istream, false, err_msg);
+			if (header.compression_type == DEFLATE) compressed_data = true;
+			else if (header.compression_type == UNCOMPRESSED) compressed_data = false;
+			else {
+				compressed_data = false; 
+#ifdef USE_UE
+#else
+				std::cerr << "ZIP: got unrecognized compressed data (Supported deflate/uncompressed)" << std::endl;
+#endif
+				valid = false;
+			}
+		}
+		else { 
+			valid = gzip_header.Read(istream, err_msg); 
+			compressed_data = true; 
+		}
+		// initialize the inflate
+		if (compressed_data && valid) {
+			int result = inflateInit2(&strm, -MAX_WBITS);
+			if (result != Z_OK) 
+			{ 
+#ifdef USE_UE
+#else
+				std::cerr << "gzip: inflateInit2 did not return Z_OK" << std::endl;
+#endif
+				valid = false; 
+			}
+		}
+	}
+
+	virtual ~ZIP_STREAMBUF_DECOMPRESS()
+	{
+		if (compressed_data && valid) inflateEnd(&strm);
+//		if (!part_of_zip_file) delete& istream;
+	}
+
+	int process()
+	{
+		if (!valid) return -1;
+		if (compressed_data) {
+			strm.avail_out = buffer_size - 4;
+			strm.next_out = (Bytef*)(out + 4);
+			while (strm.avail_out != 0) {
+				if (strm.avail_in == 0) { // buffer empty, read some more from file
+#ifdef USE_UE
+					int count = part_of_zip_file ? FGenericPlatformMath::Min((unsigned int)buffer_size, header.compressed_size - total_read) : (unsigned int)buffer_size;
+					istream->Read(reinterpret_cast<uint8*>(in), count);
+					strm.avail_in = count;
+#else
+					istream.read((char*)in, part_of_zip_file ? std::min((unsigned int)buffer_size, header.compressed_size - total_read) : (unsigned int)buffer_size);
+					strm.avail_in = (uInt)istream.gcount();
+#endif
+					total_read += strm.avail_in;
+					strm.next_in = (Bytef*)in;
+				}
+				int ret = inflate(&strm, Z_NO_FLUSH); // decompress via zlib
+				switch (ret) {
+				case Z_STREAM_ERROR:
+#ifdef USE_UE
+#else
+					std::cerr << "libz error Z_STREAM_ERROR" << std::endl;
+#endif
+					valid = false; return -1;
+				case Z_NEED_DICT:
+				case Z_DATA_ERROR:
+				case Z_MEM_ERROR:
+#ifdef USE_UE
+#else
+					std::cerr << "gzip error " << strm.msg << std::endl;
+#endif
+					valid = false; return -1;
+				}
+				if (ret == Z_STREAM_END) break;
+			}
+			int unzip_count = buffer_size - strm.avail_out - 4;
+			total_uncompressed += unzip_count;
+			return unzip_count;
+		}
+		else { // uncompressed, so just read
+#ifdef USE_UE
+			int count = std::min(buffer_size - 4, header.uncompressed_size - total_read);
+			istream->Read(reinterpret_cast<uint8*>(out + 4), count);
+#else
+			istream.read((char*)(out + 4), std::min(buffer_size - 4, header.uncompressed_size - total_read));
+			int count = (int)istream.gcount();
+#endif
+			total_read += count;
+			return count;
+		}
+		return 1;
+	}
+
+	virtual int underflow()
+	{
+		if (gptr() && (gptr() < egptr())) return traits_type::to_int_type(*gptr()); // if we already have data just use it
+		int put_back_count = (int)(gptr() - eback());
+		if (put_back_count > 4) put_back_count = 4;
+		std::memmove(out + (4 - put_back_count), gptr() - put_back_count, put_back_count);
+		int num = process();
+		setg((char*)(out + 4 - put_back_count), (char*)(out + 4), (char*)(out + 4 + num));
+		if (num <= 0) return EOF;
+		return traits_type::to_int_type(*gptr());
+	}
+
+	virtual int overflow(int c = EOF)
+	{
+#ifdef USE_UE
+#else
+		assert(false);
+#endif
+		return EOF;
+	}
+};
+
+// Class needed because istream cannot own its streambuf
+class ZIP_FILE_ISTREAM : public std::istream
+{
+	ZIP_STREAMBUF_DECOMPRESS buf;
+public:
+#ifdef USE_UE
+	ZIP_FILE_ISTREAM(IFileHandle* istream, bool part_of_zip_file)
+		: std::istream(&buf)
+		, buf(istream, part_of_zip_file)
+	{}
+#else
+	ZIP_FILE_ISTREAM(std::istream& istream, bool part_of_zip_file)
+		: std::istream(&buf)
+		, buf(istream, part_of_zip_file)
+	{}
+#endif
+
+	virtual ~ZIP_FILE_ISTREAM()
+	{}
+};
+
+class ZIP_FILE_READER
+{
+#ifdef USE_UE
+	IFileHandle* istream;
+#else
+	std::ifstream istream;
+#endif
+public:
+#ifdef USE_UE
+	TMap<FString, std::shared_ptr<ZIP_FILE_HEADER>> filename_to_header;
+#else
+	std::unordered_map<std::string, std::shared_ptr<ZIP_FILE_HEADER>> filename_to_header;
+#endif
+
+#ifdef USE_UE
+	ZIP_FILE_READER(const FString& filename)
+		: istream(nullptr)
+#else
+	ZIP_FILE_READER(const std::string& filename)
+#endif
+	{
+#ifdef USE_UE
+		FPlatformFileManager& FileManager = FPlatformFileManager::Get();
+		IPlatformFile& PlatformFile = FileManager.GetPlatformFile();
+		istream = PlatformFile.OpenRead(*filename, false);
+#else
+		istream.open(filename.c_str(), std::ios::in | std::ios::binary);
+#endif
+		if (!istream)
+		{
+	#if PLATFORM_EXCEPTIONS_DISABLED
+			return;
+	#else
+			throw std::runtime_error("ZIP: Invalid file handle");
+	#endif
+		}
+
+#ifdef USE_UE
+		FString err_msg;
+#else
+		std::string err_msg;
+#endif
+		Find_And_Read_Central_Header(err_msg);
+	}
+
+	virtual ~ZIP_FILE_READER()
+	{}
+
+#ifdef USE_UE
+	std::istream* Get_File(const FString& filename, const bool binary = true)
+	{
+		if (std::shared_ptr<ZIP_FILE_HEADER>* header = filename_to_header.Find(filename))
+		{
+			istream->Seek((*header)->header_offset);
+			return new ZIP_FILE_ISTREAM(istream, true);
+		}
+		return nullptr;
+	}
+#else
+	std::istream* Get_File(const std::string& filename, const bool binary = true)
+	{
+		//ZIP_FILE_HEADER** header = filename_to_header.Get_Pointer(filename);
+		//auto it=filename_to_header.find(filename);
+		auto it=filename_to_header.Find(FString(filename.c_str()));
+		//if(it != filename_to_header.end())
+		if(it != nullptr)
+		{
+			//ZIP_FILE_HEADER** header = &filename_to_header[filename].get();
+			//if(ZIP_FILE_HEADER* header = it->second.get())
+			if(ZIP_FILE_HEADER* header = it->get())
+			{ 
+				istream.seekg(header->header_offset); 
+				return new ZIP_FILE_ISTREAM(istream, true); 
+			}
+		}
+		return nullptr;
+	}
+#endif
+
+#ifdef USE_UE
+	void Get_File_List(TArray<FString>& filenames) const
+	{
+		filename_to_header.GetKeys(filenames);
+	}
+#else
+	void Get_File_List(std::vector<std::string>& filenames) const
+	{
+		filenames.reserve(filename_to_header.size());
+		for(const auto& kv : filename_to_header)
+			filenames.push_back(kv.first);
+	}
+#endif
+
+private:
+
+#ifdef USE_UE
+	template<class T>
+	void
+	Read_Primitive(IFileHandle* input, T& d)
+	{
+		//STATIC_ASSERT((sizeof(T) == PLATFORM_INDEPENDENT_SIZE<T>::value));
+		input->Read(reinterpret_cast<uint8*>(&d), sizeof(T));
+		//if (big_endian) Swap_Endianity(d); // convert to little endian if necessary
+	}
+#else
+	template<class T>
+	void
+	Read_Primitive(std::istream& input, T& d)
+	{
+		//STATIC_ASSERT((sizeof(T) == PLATFORM_INDEPENDENT_SIZE<T>::value));
+		input.read((char*)&d, sizeof(T));
+		//if (big_endian) Swap_Endianity(d); // convert to little endian if necessary
+	}
+#endif
+
+#ifdef USE_UE
+	bool
+	Find_And_Read_Central_Header(FString& err_msg)
+	{
+#else
+	bool
+	Find_And_Read_Central_Header(std::string& err_msg)
+	{
+#endif
+		// Find the header
+		// NOTE: this assumes the zip file header is the last thing written to file...
+#ifdef USE_UE
+		int end_position = istream->Size();
+#else
+		istream.seekg(0, std::ios_base::end);
+		int end_position = istream.tellg();
+#endif
+		unsigned int max_comment_size = 0xffff; // max size of header
+		unsigned int read_size_before_comment = 22;
+		unsigned int read_start = max_comment_size + read_size_before_comment;
+		if (static_cast<int>(read_start) > end_position) read_start = end_position;
+#ifdef USE_UE
+		istream->Seek(end_position - read_start);
+#else
+		istream.seekg(end_position - read_start);
+#endif
+		if (read_start <= 0) 
+		{ 
+			//LOG::cerr << "ZIP: Invalid read buffer size" << std::endl;
+			err_msg += "ZIP_FILE_READER: Invalid read buffer size.\n";
+			return false; 
+		}
+		char* buf = new char[(unsigned int)read_start];
+#ifdef USE_UE
+		istream->Read(reinterpret_cast<uint8*>(buf), read_start);
+#else
+		istream.read(buf, read_start);
+#endif
+		int found = -1;
+		for (unsigned int i = 0; i < read_start - 3; i++) {
+			if (buf[i] == 0x50 && buf[i + 1] == 0x4b && buf[i + 2] == 0x05 && buf[i + 3] == 0x06) 
+			{ 
+				found = i; 
+				break; 
+			}
+		}
+		delete[] buf;
+		if (found == -1) 
+		{ 
+			//LOG::cerr << "ZIP: Failed to find zip header" << std::endl;
+			err_msg += "ZIP_FILE_READER: Failed to find zip header.\n";
+			return false; 
+		}
+		// seek to end of central header and read
+#ifdef USE_UE
+		istream->Seek(end_position - (read_start - found));
+#else
+		istream.seekg(end_position - (read_start - found));
+#endif
+		unsigned int word;
+		unsigned short disk_number1, disk_number2, num_files, num_files_this_disk;
+		Read_Primitive(istream, word); // end of central
+		Read_Primitive(istream, disk_number1); // this disk number
+		Read_Primitive(istream, disk_number2); // this disk number
+		if (disk_number1 != disk_number2 || disk_number1 != 0) {
+			//LOG::cerr << "ZIP: multiple disk zip files are not supported" << std::endl;
+			err_msg += "ZIP_FILE_READER: multiple disk zip files are not supported.\n";
+			return false;
+		}
+		Read_Primitive(istream, num_files); // one entry in center in this disk
+		Read_Primitive(istream, num_files_this_disk); // one entry in center 
+		if (num_files != num_files_this_disk) {
+			//LOG::cerr << "ZIP: multi disk zip files are not supported" << std::endl;
+			err_msg += "ZIP: multi disk zip files are not supported.\n";
+			return false;
+		}
+		unsigned int size_of_header, header_offset;
+		Read_Primitive(istream, size_of_header); // size of header
+		Read_Primitive(istream, header_offset); // offset to header
+		// go to header and read all file headers
+#ifdef USE_UE
+		istream->Seek(header_offset);
+#else
+		istream.seekg(header_offset);
+#endif
+		for (int i = 0; i < num_files; i++) 
+		{
+			std::shared_ptr<ZIP_FILE_HEADER> header(new ZIP_FILE_HEADER);
+			if (header->Read(istream, true, err_msg))
+			{
+#ifdef USE_UE
+				filename_to_header.Add(header->filename, header);
+#else
+				filename_to_header.insert(std::make_pair(header->filename, header));
+#endif
+			}
+		}
+		return true;
+	}
+
+};
+
+#endif // WITH_EDITOR
+#endif // USE_ZLIB
diff --git a/Engine/Plugins/Experimental/ChaosFlesh/Source/ChaosFlesh/Public/ChaosFlesh/FleshCollectionUtility.h b/Engine/Plugins/Experimental/ChaosFlesh/Source/ChaosFlesh/Public/ChaosFlesh/FleshCollectionUtility.h
new file mode 100644
--- /dev/null
+++ b/Engine/Plugins/Experimental/ChaosFlesh/Source/ChaosFlesh/Public/ChaosFlesh/FleshCollectionUtility.h
@@ -0,0 +1,23 @@
+// Copyright Epic Games, Inc. All Rights Reserved.
+
+#pragma once
+
+#include "ChaosFlesh/FleshCollection.h"
+#include "CoreMinimal.h"
+#include "Templates/UniquePtr.h"
+#include "Misc/AssertionMacros.h"
+
+DECLARE_LOG_CATEGORY_EXTERN(LogChaosFlesh, Verbose, All);
+
+namespace ChaosFlesh 
+{
+	void GetTetFaces(const FIntVector4& Tet, FIntVector3& Face1, FIntVector3& Face2, FIntVector3& Face3, FIntVector3& Face4, const bool invert = false);
+	void GetSurfaceElements(const TArray<FIntVector4>& Tets, TArray<FIntVector3>& SurfaceElements, const bool KeepInteriorFaces, const bool InvertFaces = false);
+
+	TUniquePtr<FFleshCollection> ReadGEOFile(const FString& Filename);
+	TUniquePtr<FFleshCollection> ReadTetFile(const FString& Filename); // uncompressed tet
+	TUniquePtr<FFleshCollection> ReadTetPBDeformableGeometryCollection(const FString& Filename);
+
+	//! Supports reading .tet and .geo (compatible with version 19).
+	TUniquePtr<FFleshCollection> CHAOSFLESH_API ImportTetFromFile(const FString& Filename);
+}
diff --git a/Engine/Plugins/Experimental/ChaosFlesh/Source/ChaosFlesh/Public/ChaosFlesh/TetrahedralCollection.h b/Engine/Plugins/Experimental/ChaosFlesh/Source/ChaosFlesh/Public/ChaosFlesh/TetrahedralCollection.h
new file mode 100644
--- /dev/null
+++ b/Engine/Plugins/Experimental/ChaosFlesh/Source/ChaosFlesh/Public/ChaosFlesh/TetrahedralCollection.h
@@ -0,0 +1,96 @@
+// Copyright Epic Games, Inc. All Rights Reserved.
+
+#pragma once
+
+#include "CoreMinimal.h"
+#include "GeometryCollection/GeometryCollection.h"
+#include "GeometryCollection/ManagedArrayCollection.h"
+#include "Misc/Crc.h"
+
+namespace Chaos
+{
+	class FChaosArchive;
+}
+
+
+/**
+* FTetrahedralCollection (FGeometryCollection)
+*/
+class CHAOSFLESH_API FTetrahedralCollection : public FGeometryCollection
+{
+	typedef FGeometryCollection Super;
+
+public:
+	FTetrahedralCollection();
+	FTetrahedralCollection(FTetrahedralCollection &) = delete;
+	FTetrahedralCollection& operator=(const FTetrahedralCollection &) = delete;
+	FTetrahedralCollection(FTetrahedralCollection &&) = default;
+	FTetrahedralCollection& operator=(FTetrahedralCollection &&) = default;
+
+	/**
+	 * Create a GeometryCollection from Vertex and Indices arrays
+	 */
+	static FTetrahedralCollection* NewTetrahedralCollection(const TArray<FVector>& Vertices, const TArray<FIntVector3>& SurfaceElements, const TArray<FIntVector4>& Elements,bool bReverseVertexOrder = true);
+	static void Init(FTetrahedralCollection* Collection, const TArray<FVector>& Vertices, const TArray<FIntVector3>& SurfaceElements, const TArray<FIntVector4>& Elements, bool bReverseVertexOrder = true);
+
+	int32 AppendGeometry(const FTetrahedralCollection& GeometryCollection, int32 MaterialIDOffset = 0, bool ReindexAllMaterials = true, const FTransform& TransformRoot = FTransform::Identity);
+
+	/**
+	* Build \c IncidentElements and \c IncidentElementsLocalIndex.
+	* \p GeometryIndex - ID of the geometry entry to restrict operation to; -1 does all.
+	*/
+	void InitIncidentElements(const int32 GeometryIndex=-1);
+
+	/*
+	*  SetDefaults for new entries on this collection. 
+	*/
+	virtual void SetDefaults(FName Group, uint32 StartSize, uint32 NumElements) override;
+
+	/**
+	* Reorders elements in a group. NewOrder must be the same length as the group.
+	*/
+	void ReorderElements(FName Group, const TArray<int32>& NewOrder) override;
+
+	void ReorderTetrahedralElements(const TArray<int32>& NewOrder);
+
+	/*
+	*  Attribute Groupes
+	*/
+	static const FName TetrahedralGroup;
+
+	/*
+	*  Tetrahedron Attribute
+	*  TManagedArray<FIntVector4> Tetrahedron = this->FindAttribute<FIntVector4>(FTetrahedralCollection::TetrahedronAttribute,FTetrahedralCollection::TetrahedralGroup);
+	*/
+	static const FName TetrahedronAttribute;
+	TManagedArray<FIntVector4> Tetrahedron;
+
+	TManagedArray<int32> TetrahedronStart;
+	TManagedArray<int32> TetrahedronCount;
+
+	/**
+	* Incident Elements Attribute
+	* For each vertex, a list of tetrahedra that includes that vertex, and the vertex's index in the tet.
+	*/
+	static const FName IncidentElementsAttribute;
+	TManagedArray<TArray<int32>> IncidentElements;
+
+	/**
+	* Incident Elements Attribute
+	* For each incident element, the vertex's index in the tetrahedron.
+	*/
+	static const FName IncidentElementsLocalIndexAttribute;
+	TManagedArray<TArray<int32>> IncidentElementsLocalIndex;
+
+protected:
+
+	void Construct();
+
+};
+
+FORCEINLINE Chaos::FChaosArchive& operator<<(Chaos::FChaosArchive& Ar, FTetrahedralCollection& Value)
+{
+	Value.Serialize(Ar);
+	return Ar;
+}
+
diff --git a/Engine/Plugins/Experimental/ChaosFlesh/Source/ChaosFleshEditor/Private/ChaosFlesh/Asset/ActorFactoryDeformableSolver.cpp b/Engine/Plugins/Experimental/ChaosFlesh/Source/ChaosFleshEditor/Private/ChaosFlesh/Asset/ActorFactoryDeformableSolver.cpp
new file mode 100644
--- /dev/null
+++ b/Engine/Plugins/Experimental/ChaosFlesh/Source/ChaosFleshEditor/Private/ChaosFlesh/Asset/ActorFactoryDeformableSolver.cpp
@@ -0,0 +1,53 @@
+// Copyright Epic Games, Inc. All Rights Reserved.
+
+#include "ChaosFlesh/Asset/ActorFactoryDeformableSolver.h"
+
+#include "ChaosFlesh/ChaosDeformableSolverActor.h"
+#include "ChaosFlesh/ChaosDeformableSolverAsset.h"
+#include "Misc/FileHelper.h"
+#include "AssetRegistry/AssetData.h"
+#include "Editor.h"
+#include "EditorFramework/AssetImportData.h"
+
+#define LOCTEXT_NAMESPACE "ActorFactoryDeformableSolver"
+
+
+/*-----------------------------------------------------------------------------
+UActorFactoryDeformableSolver
+-----------------------------------------------------------------------------*/
+UActorFactoryDeformableSolver::UActorFactoryDeformableSolver(const FObjectInitializer& ObjectInitializer)
+	: Super(ObjectInitializer)
+{
+	DisplayName = LOCTEXT("DeformableSolverDisplayName", "DeformableSolver");
+	NewActorClass = ADeformableSolverActor::StaticClass();
+}
+
+bool UActorFactoryDeformableSolver::CanCreateActorFrom(const FAssetData& AssetData, FText& OutErrorMsg)
+{
+	if (!AssetData.IsValid() || !AssetData.IsInstanceOf(UChaosDeformableSolver::StaticClass()))
+	{
+		OutErrorMsg = NSLOCTEXT("CanCreateActor", "NoDeformableSolverSpecified", "No solver asset was specified.");
+		return false;
+	}
+
+	return true;
+}
+
+void UActorFactoryDeformableSolver::PostSpawnActor(UObject* Asset, AActor* NewActor)
+{
+	Super::PostSpawnActor(Asset, NewActor);
+
+	UChaosDeformableSolver* ChaosSolver = CastChecked<UChaosDeformableSolver>(Asset);
+	ADeformableSolverActor * NewDeformableSolverActo = CastChecked<ADeformableSolverActor >(NewActor);
+}
+
+void UActorFactoryDeformableSolver::PostCreateBlueprint(UObject* Asset, AActor* CDO)
+{
+	if (Asset != NULL && CDO != NULL)
+	{
+		UChaosDeformableSolver* ChaosSolver = CastChecked<UChaosDeformableSolver>(Asset);
+		ADeformableSolverActor* ChaosSolverActor = CastChecked<ADeformableSolverActor>(CDO);
+	}
+}
+
+#undef LOCTEXT_NAMESPACE
\ No newline at end of file
diff --git a/Engine/Plugins/Experimental/ChaosFlesh/Source/ChaosFleshEditor/Private/ChaosFlesh/Asset/ActorFactoryDeformableSolver.h b/Engine/Plugins/Experimental/ChaosFlesh/Source/ChaosFleshEditor/Private/ChaosFlesh/Asset/ActorFactoryDeformableSolver.h
new file mode 100644
--- /dev/null
+++ b/Engine/Plugins/Experimental/ChaosFlesh/Source/ChaosFleshEditor/Private/ChaosFlesh/Asset/ActorFactoryDeformableSolver.h
@@ -0,0 +1,27 @@
+// Copyright Epic Games, Inc. All Rights Reserved.
+
+#pragma once
+
+#include "CoreMinimal.h"
+#include "UObject/ObjectMacros.h"
+#include "ActorFactories/ActorFactory.h"
+
+#include "ActorFactoryDeformableSolver.generated.h"
+
+class AActor;
+struct FAssetData;
+
+UCLASS(MinimalAPI, config=Editor)
+class UActorFactoryDeformableSolver : public UActorFactory
+{
+	GENERATED_UCLASS_BODY()
+
+	//~ Begin UActorFactory Interface
+	virtual void PostSpawnActor( UObject* Asset, AActor* NewActor ) override;
+	virtual void PostCreateBlueprint(UObject* Asset, AActor* CDO) override;
+	virtual bool CanCreateActorFrom(const FAssetData& AssetData, FText& OutErrorMsg) override;
+	//~ End UActorFactory Interface
+};
+
+
+
diff --git a/Engine/Plugins/Experimental/ChaosFlesh/Source/ChaosFleshEditor/Private/ChaosFlesh/Asset/AssetTypeActions_ChaosDeformableSolverAsset.cpp b/Engine/Plugins/Experimental/ChaosFlesh/Source/ChaosFleshEditor/Private/ChaosFlesh/Asset/AssetTypeActions_ChaosDeformableSolverAsset.cpp
new file mode 100644
--- /dev/null
+++ b/Engine/Plugins/Experimental/ChaosFlesh/Source/ChaosFleshEditor/Private/ChaosFlesh/Asset/AssetTypeActions_ChaosDeformableSolverAsset.cpp
@@ -0,0 +1,31 @@
+// Copyright Epic Games, Inc. All Rights Reserved.
+
+#include "ChaosFlesh/Asset/AssetTypeActions_ChaosDeformableSolverAsset.h"
+
+#include "ThumbnailRendering/SceneThumbnailInfo.h"
+#include "ChaosFlesh/ChaosDeformableSolverAsset.h"
+#include "ToolMenus.h"
+
+#define LOCTEXT_NAMESPACE "AssetTypeActions"
+
+UClass* FAssetTypeActions_ChaosDeformableSolver::GetSupportedClass() const
+{
+	return UChaosDeformableSolver::StaticClass();
+}
+
+UThumbnailInfo* FAssetTypeActions_ChaosDeformableSolver::GetThumbnailInfo(UObject* Asset) const
+{
+	UChaosDeformableSolver* ChaosDeformableSolver = CastChecked<UChaosDeformableSolver>(Asset);
+	return NewObject<USceneThumbnailInfo>(ChaosDeformableSolver, NAME_None, RF_Transactional);
+}
+
+void FAssetTypeActions_ChaosDeformableSolver::GetActions(const TArray<UObject*>& InObjects, FToolMenuSection& Section)
+{
+	FAssetTypeActions_Base::GetActions(InObjects, Section);
+}
+
+void FAssetTypeActions_ChaosDeformableSolver::OpenAssetEditor(const TArray<UObject*>& InObjects, TSharedPtr<IToolkitHost> EditWithinLevelEditor)
+{
+}
+
+#undef LOCTEXT_NAMESPACE
diff --git a/Engine/Plugins/Experimental/ChaosFlesh/Source/ChaosFleshEditor/Private/ChaosFlesh/Asset/AssetTypeActions_FleshAsset.cpp b/Engine/Plugins/Experimental/ChaosFlesh/Source/ChaosFleshEditor/Private/ChaosFlesh/Asset/AssetTypeActions_FleshAsset.cpp
new file mode 100644
--- /dev/null
+++ b/Engine/Plugins/Experimental/ChaosFlesh/Source/ChaosFleshEditor/Private/ChaosFlesh/Asset/AssetTypeActions_FleshAsset.cpp
@@ -0,0 +1,48 @@
+// Copyright Epic Games, Inc. All Rights Reserved.
+
+#include "ChaosFlesh/Asset/AssetTypeActions_FleshAsset.h"
+
+#include "ChaosFlesh/FleshAsset.h"
+#include "ChaosFlesh/ChaosFleshEditorPlugin.h"
+#include "ThumbnailRendering/SceneThumbnailInfo.h"
+#include "ToolMenus.h"
+
+#define LOCTEXT_NAMESPACE "AssetTypeActions"
+
+UClass* FAssetTypeActions_FleshAsset::GetSupportedClass() const
+{
+	return UFleshAsset::StaticClass();
+}
+
+UThumbnailInfo* FAssetTypeActions_FleshAsset::GetThumbnailInfo(UObject* Asset) const
+{
+	UFleshAsset * FleshAsset = CastChecked<UFleshAsset>(Asset);
+	UThumbnailInfo* ThumbnailInfo = FleshAsset->ThumbnailInfo;
+	if (ThumbnailInfo == NULL)
+	{
+		ThumbnailInfo = NewObject<USceneThumbnailInfo>(FleshAsset, NAME_None, RF_Transactional);
+		FleshAsset->ThumbnailInfo = ThumbnailInfo;
+	}
+
+	return ThumbnailInfo;
+}
+
+void FAssetTypeActions_FleshAsset::GetActions(const TArray<UObject*>& InObjects, FToolMenuSection& Section)
+{
+	FAssetTypeActions_Base::GetActions(InObjects, Section);
+}
+
+void FAssetTypeActions_FleshAsset::OpenAssetEditor(const TArray<UObject*>& InObjects, TSharedPtr<class IToolkitHost> EditWithinLevelEditor)
+{
+	EToolkitMode::Type Mode = EditWithinLevelEditor.IsValid() ? EToolkitMode::WorldCentric : EToolkitMode::Standalone;
+	for (auto ObjIt = InObjects.CreateConstIterator(); ObjIt; ++ObjIt)
+	{
+		if (auto Object = Cast<UFleshAsset>(*ObjIt))
+		{
+			IChaosFleshEditorPlugin* FleshEditorModule = &FModuleManager::LoadModuleChecked<IChaosFleshEditorPlugin>("ChaosFleshEditor");
+			FleshEditorModule->CreateFleshAssetEditor(Mode, EditWithinLevelEditor, Object);
+		}
+	}
+}
+
+#undef LOCTEXT_NAMESPACE
diff --git a/Engine/Plugins/Experimental/ChaosFlesh/Source/ChaosFleshEditor/Private/ChaosFlesh/Asset/ChaosAssetDeformableSolverFactory.cpp b/Engine/Plugins/Experimental/ChaosFlesh/Source/ChaosFleshEditor/Private/ChaosFlesh/Asset/ChaosAssetDeformableSolverFactory.cpp
new file mode 100644
--- /dev/null
+++ b/Engine/Plugins/Experimental/ChaosFlesh/Source/ChaosFleshEditor/Private/ChaosFlesh/Asset/ChaosAssetDeformableSolverFactory.cpp
@@ -0,0 +1,37 @@
+// Copyright Epic Games, Inc. All Rights Reserved.
+
+#include "ChaosFlesh/Asset/ChaosAssetDeformableSolverFactory.h"
+
+#include "Editor.h"
+#include "Editor/EditorEngine.h"
+#include "Engine/Selection.h"
+
+#define LOCTEXT_NAMESPACE "ChaosDeformableSolver"
+
+/////////////////////////////////////////////////////
+// ChaosDeformableSolverFactory
+
+UChaosDeformableSolverFactory::UChaosDeformableSolverFactory(const FObjectInitializer& ObjectInitializer)
+	: Super(ObjectInitializer)
+{
+	bCreateNew = true;
+	bEditAfterNew = true;
+	SupportedClass = UChaosDeformableSolver::StaticClass();
+}
+
+UChaosDeformableSolver* UChaosDeformableSolverFactory::StaticFactoryCreateNew(UClass* Class, UObject* InParent, FName Name, EObjectFlags Flags, UObject* Context, FFeedbackContext* Warn)
+{
+	return static_cast<UChaosDeformableSolver*>(NewObject<UChaosDeformableSolver>(InParent, Class, Name, Flags | RF_Transactional | RF_Public | RF_Standalone));
+}
+
+UObject* UChaosDeformableSolverFactory::FactoryCreateNew(UClass* Class, UObject* InParent, FName Name, EObjectFlags Flags, UObject* Context, FFeedbackContext* Warn)
+{
+	UChaosDeformableSolver* NewChaosDeformableSolver = StaticFactoryCreateNew(Class, InParent, Name, Flags, Context, Warn);
+	NewChaosDeformableSolver->MarkPackageDirty();
+	return NewChaosDeformableSolver;
+}
+
+#undef LOCTEXT_NAMESPACE
+
+
+
diff --git a/Engine/Plugins/Experimental/ChaosFlesh/Source/ChaosFleshEditor/Private/ChaosFlesh/Asset/FleshAssetFactory.cpp b/Engine/Plugins/Experimental/ChaosFlesh/Source/ChaosFleshEditor/Private/ChaosFlesh/Asset/FleshAssetFactory.cpp
new file mode 100644
--- /dev/null
+++ b/Engine/Plugins/Experimental/ChaosFlesh/Source/ChaosFleshEditor/Private/ChaosFlesh/Asset/FleshAssetFactory.cpp
@@ -0,0 +1,55 @@
+// Copyright Epic Games, Inc. All Rights Reserved.
+
+#include "ChaosFlesh/Asset/FleshAssetFactory.h"
+
+#include "ChaosFlesh/FleshAsset.h"
+#include "ChaosFlesh/FleshComponent.h"
+#include "Editor.h"
+#include "Editor/EditorEngine.h"
+#include "Engine/Selection.h"
+
+#include UE_INLINE_GENERATED_CPP_BY_NAME(FleshAssetFactory)
+
+
+#define LOCTEXT_NAMESPACE "ChaosFlesh"
+
+/////////////////////////////////////////////////////
+// FleshFactory
+
+UFleshAssetFactory::UFleshAssetFactory(const FObjectInitializer& ObjectInitializer)
+	: Super(ObjectInitializer)
+{
+	bCreateNew = true;
+	bEditAfterNew = true;
+	SupportedClass = UFleshAsset::StaticClass();
+}
+
+UObject* UFleshAssetFactory::FactoryCreateNew(UClass* Class, UObject* InParent, FName Name, EObjectFlags Flags, UObject* Context, FFeedbackContext* Warn)
+{
+	FTransform LastTransform = FTransform::Identity;
+	TArray< FleshTuple > FleshList;
+
+	TArray<FAssetData> SelectedAssets;
+	GEditor->GetContentBrowserSelections(SelectedAssets);
+
+	for (const FAssetData& AssetData : SelectedAssets)
+	{
+		if (AssetData.GetAsset()->IsA<UFleshAsset>())
+		{
+			UFleshComponent *DummyValue(NULL);
+			FleshList.Add(FleshTuple(static_cast<const UFleshAsset *>(AssetData.GetAsset()), DummyValue, FTransform()));
+		}
+
+	}
+
+	UFleshAsset* NewFlesh = NewObject<UFleshAsset>(InParent, Class, Name, Flags | RF_Transactional | RF_Public | RF_Standalone);
+
+	NewFlesh->Modify();
+	return NewFlesh;
+}
+
+#undef LOCTEXT_NAMESPACE
+
+
+
+
diff --git a/Engine/Plugins/Experimental/ChaosFlesh/Source/ChaosFleshEditor/Private/ChaosFlesh/ChaosFleshEditorPlugin.cpp b/Engine/Plugins/Experimental/ChaosFlesh/Source/ChaosFleshEditor/Private/ChaosFlesh/ChaosFleshEditorPlugin.cpp
new file mode 100644
--- /dev/null
+++ b/Engine/Plugins/Experimental/ChaosFlesh/Source/ChaosFleshEditor/Private/ChaosFlesh/ChaosFleshEditorPlugin.cpp
@@ -0,0 +1,79 @@
+// Copyright Epic Games, Inc. All Rights Reserved.
+
+
+#include "ChaosFlesh/ChaosFleshEditorPlugin.h"
+#include "Editor/FleshEditorStyle.h"
+
+#include "AssetToolsModule.h"
+#include "ChaosFlesh/FleshComponent.h"
+#include "ChaosFlesh/Cmd/FleshAssetConversion.h"
+#include "ChaosFlesh/Cmd/ChaosFleshCommands.h"
+#include "ChaosFlesh/Asset/AssetTypeActions_FleshAsset.h"
+#include "ChaosFlesh/Asset/AssetTypeActions_ChaosDeformableSolverAsset.h"
+#include "ChaosFlesh/Asset/FleshAssetThumbnailRenderer.h"
+#include "ChaosFlesh/Editor/FleshEditorToolkit.h"
+#include "CoreMinimal.h"
+#include "Features/IModularFeatures.h"
+#include "HAL/ConsoleManager.h"
+#include "LevelEditor.h"
+#include "Modules/ModuleManager.h"
+#include "PropertyEditorModule.h"
+#include "Styling/SlateStyle.h"
+#include "Styling/SlateStyleRegistry.h"
+#include "Styling/CoreStyle.h"
+#include "ToolMenus.h"
+#include "Templates/SharedPointer.h"
+
+#define LOCTEXT_NAMESPACE "FleshEditor"
+
+#define BOX_BRUSH(StyleSet, RelativePath, ...) FSlateBoxBrush(StyleSet->RootToContentDir(RelativePath, TEXT(".png")), __VA_ARGS__)
+#define DEFAULT_FONT(...) FCoreStyle::GetDefaultFontStyle(__VA_ARGS__)
+
+void IChaosFleshEditorPlugin::StartupModule()
+{
+	FChaosFleshEditorStyle::Get();
+
+	IAssetTools& AssetTools = FAssetToolsModule::GetModule().Get();
+	FleshAssetActions = new FAssetTypeActions_FleshAsset();
+	AssetTools.RegisterAssetTypeActions(MakeShareable(FleshAssetActions));
+	ChaosDeformableSolverAssetActions = new FAssetTypeActions_ChaosDeformableSolver();
+	AssetTools.RegisterAssetTypeActions(MakeShareable(ChaosDeformableSolverAssetActions));
+
+
+	if (GIsEditor && !IsRunningCommandlet())
+	{
+		EditorCommands.Add(IConsoleManager::Get().RegisterConsoleCommand(
+			TEXT("p.Chaos.Flesh.ImportFile"),
+			TEXT("Creates a FleshAsset from the input file"),
+			FConsoleCommandWithWorldAndArgsDelegate::CreateStatic(&FChaosFleshCommands::ImportFile),
+			ECVF_Default
+		));
+	}
+
+	UThumbnailManager::Get().RegisterCustomRenderer(UFleshAsset::StaticClass(), UFleshAssetThumbnailRenderer::StaticClass());
+
+}
+
+void IChaosFleshEditorPlugin::ShutdownModule()
+{
+	if (UObjectInitialized())
+	{	
+		UThumbnailManager::Get().UnregisterCustomRenderer(UFleshAsset::StaticClass());
+
+		IAssetTools& AssetTools = FAssetToolsModule::GetModule().Get();
+		AssetTools.UnregisterAssetTypeActions(FleshAssetActions->AsShared());
+		AssetTools.UnregisterAssetTypeActions(ChaosDeformableSolverAssetActions->AsShared());
+	}
+}
+
+TSharedRef<FAssetEditorToolkit> IChaosFleshEditorPlugin::CreateFleshAssetEditor(const EToolkitMode::Type Mode, const TSharedPtr<IToolkitHost>& InitToolkitHost, UObject* FleshAsset)
+{
+	TSharedPtr<FFleshEditorToolkit> NewFleshAssetEditor = MakeShared<FFleshEditorToolkit>();
+	NewFleshAssetEditor->InitFleshAssetEditor(Mode, InitToolkitHost, FleshAsset);
+	return StaticCastSharedPtr<FAssetEditorToolkit>(NewFleshAssetEditor).ToSharedRef();
+}
+
+IMPLEMENT_MODULE(IChaosFleshEditorPlugin, FleshAssetEditor)
+
+
+#undef LOCTEXT_NAMESPACE
diff --git a/Engine/Plugins/Experimental/ChaosFlesh/Source/ChaosFleshEditor/Private/ChaosFlesh/Cmd/ChaosFleshCommands.cpp b/Engine/Plugins/Experimental/ChaosFlesh/Source/ChaosFleshEditor/Private/ChaosFlesh/Cmd/ChaosFleshCommands.cpp
new file mode 100644
--- /dev/null
+++ b/Engine/Plugins/Experimental/ChaosFlesh/Source/ChaosFleshEditor/Private/ChaosFlesh/Cmd/ChaosFleshCommands.cpp
@@ -0,0 +1,58 @@
+// Copyright Epic Games, Inc. All Rights Reserved.
+
+#include "ChaosFlesh/Cmd/ChaosFleshCommands.h"
+
+#include "ChaosFlesh/Asset/FleshAssetFactory.h"
+#include "ChaosFlesh/FleshAsset.h"
+#include "ChaosFlesh/FleshActor.h"
+#include "ChaosFlesh/FleshComponent.h"
+#include "ChaosFlesh/FleshCollectionUtility.h"
+
+#include "AssetRegistry/AssetRegistryModule.h"
+#include "Editor.h"
+#include "EditorSupportDelegates.h"
+#include "Engine/Selection.h"
+#include "Engine/World.h"
+#include "EngineUtils.h"
+#include "Internationalization/Regex.h"
+#include "Logging/LogMacros.h"
+#include "Misc/Paths.h"
+#include "SceneOutlinerDelegates.h"
+
+
+DEFINE_LOG_CATEGORY_STATIC(UChaosFleshCommandsLogging, NoLogging, All);
+
+
+
+void FChaosFleshCommands::ImportFile(const TArray<FString>& Args, UWorld* World)
+{
+	if (Args.Num() == 1)
+	{
+		//FString BasePath = FPaths::ProjectDir() + "Imports/default.geo";// +FString(Args[0]);
+		if (FPaths::FileExists(Args[0]))
+		{
+			auto Factory = NewObject<UFleshAssetFactory>();
+			UPackage* Package = CreatePackage(TEXT("/Game/FleshAsset"));
+			//UPackage* Package = CreatePackage( *FPaths::ProjectContentDir() );
+
+			UFleshAsset* FleshAsset = static_cast<UFleshAsset*>(Factory->FactoryCreateNew(UFleshAsset::StaticClass(), Package, FName("FleshAsset"), RF_Standalone | RF_Public, NULL, GWarn));
+			FAssetRegistryModule::AssetCreated(FleshAsset);
+			{
+				FFleshAssetEdit EditObject = FleshAsset->EditCollection();
+				if (FFleshCollection* Collection = EditObject.GetFleshCollection())
+				{
+					UE_LOG(UChaosFleshCommandsLogging, Log, TEXT("FChaosFleshCommands::ImportFile"));
+					if (TUniquePtr<FFleshCollection> InCollection = ChaosFlesh::ImportTetFromFile(Args[0]))
+					{
+						Collection->CopyMatchingAttributesFrom(*InCollection);
+					}
+				}
+				Package->SetDirtyFlag(true);
+			}
+		}
+	}
+	else
+	{
+		UE_LOG(UChaosFleshCommandsLogging, Error, TEXT("Failed to import file for flesh asset."));
+	}
+}
diff --git a/Engine/Plugins/Experimental/ChaosFlesh/Source/ChaosFleshEditor/Private/ChaosFlesh/Cmd/FleshAssetConversion.cpp b/Engine/Plugins/Experimental/ChaosFlesh/Source/ChaosFleshEditor/Private/ChaosFlesh/Cmd/FleshAssetConversion.cpp
new file mode 100644
--- /dev/null
+++ b/Engine/Plugins/Experimental/ChaosFlesh/Source/ChaosFleshEditor/Private/ChaosFlesh/Cmd/FleshAssetConversion.cpp
@@ -0,0 +1,17 @@
+// Copyright Epic Games, Inc. All Rights Reserved.
+
+#include "ChaosFlesh/Cmd/FleshAssetConversion.h"
+
+#include "ChaosFlesh/FleshAsset.h"
+#include "ChaosFlesh/FleshActor.h"
+#include "ChaosFlesh/Asset/FleshAssetFactory.h"
+#include "ChaosFlesh/FleshComponent.h"
+#include "Logging/LogMacros.h"
+
+
+DEFINE_LOG_CATEGORY_STATIC(UFleshAssetConversionLogging, Log, All);
+
+
+void FFleshAssetConversion::ImportFile(FString & FileName)
+{
+}
\ No newline at end of file
