ChaosFlesh : Moved to Experimental

diff --git a/Engine/Plugins/Experimental/ChaosFlesh/Source/ChaosFleshEditor/Private/ChaosFlesh/Editor/FleshEditorToolkit.cpp b/Engine/Plugins/Experimental/ChaosFlesh/Source/ChaosFleshEditor/Private/ChaosFlesh/Editor/FleshEditorToolkit.cpp
new file mode 100644
--- /dev/null
+++ b/Engine/Plugins/Experimental/ChaosFlesh/Source/ChaosFleshEditor/Private/ChaosFlesh/Editor/FleshEditorToolkit.cpp
@@ -0,0 +1,410 @@
+// Copyright Epic Games, Inc. All Rights Reserved.
+
+#include "FleshEditorToolkit.h"
+
+#include "ChaosFlesh/Asset/ActorFactoryFlesh.h"
+#include "ChaosFlesh/FleshAsset.h"
+#include "Dataflow/DataflowEditor.h"
+#include "Dataflow/DataflowEditorActions.h"
+#include "Dataflow/DataflowGraphEditor.h"
+#include "Dataflow/DataflowNodeParameters.h"
+#include "EditorViewportTabContent.h"
+#include "EditorViewportLayout.h"
+#include "EditorViewportCommands.h"
+#include "Animation/Skeleton.h"
+#include "Engine/SkeletalMesh.h"
+#include "Engine/World.h"
+#include "FleshEditorViewport.h"
+#include "Framework/Commands/GenericCommands.h"
+#include "GraphEditor.h"
+#include "GraphEditorActions.h"
+#include "Engine/DataTable.h"
+
+#include "Modules/ModuleManager.h"
+#include "PropertyEditorModule.h"
+#include "Styling/AppStyle.h"
+#include "Styling/SlateStyleRegistry.h"
+
+#include "ISkeletonTree.h"
+#include "ISkeletonEditorModule.h"
+#include "IStructureDetailsView.h"
+
+DEFINE_LOG_CATEGORY_STATIC(LogChaosFleshEditorToolkit, Log, All);
+
+int32 bDataflowAssetEditorFleshLiveEvaluationEnableCVar = 1;
+FAutoConsoleVariableRef CVarDataflowAssetEditorFleshEnableLiveEvaluation(TEXT("p.Dataflow.AssetEditor.Flesh.LiveEvaluation.Enable"),
+	bDataflowAssetEditorFleshLiveEvaluationEnableCVar,
+	TEXT("Enable live evaluation of specified output on the FleshAsset within the Dataflow Editor.[def:true]"));
+
+
+#define LOCTEXT_NAMESPACE "FleshEditorToolkit"
+
+const FName FFleshEditorToolkit::ViewportTabId(TEXT("FleshEditor_Viewport"));
+const FName FFleshEditorToolkit::GraphCanvasTabId(TEXT("FleshEditor_GraphCanvas"));
+const FName FFleshEditorToolkit::AssetDetailsTabId(TEXT("FleshEditor_AssetDetails"));
+const FName FFleshEditorToolkit::NodeDetailsTabId(TEXT("FleshEditor_NodeDetails"));
+const FName FFleshEditorToolkit::SkeletalTabId(TEXT("FleshEditor_Skeletal"));
+
+
+void FFleshEditorToolkit::InitFleshAssetEditor(const EToolkitMode::Type Mode, const TSharedPtr<IToolkitHost>& InitToolkitHost, UObject* ObjectToEdit)
+{
+	FleshAsset = CastChecked<UFleshAsset>(ObjectToEdit);
+	if (FleshAsset != nullptr)
+	{
+		if (FleshAsset->Dataflow == nullptr)
+		{
+			const FName NodeName = MakeUniqueObjectName(FleshAsset, UDataflow::StaticClass(), FName("DataflowFleshAsset"));
+			FleshAsset->Dataflow = NewObject<UDataflow>(FleshAsset, NodeName);
+		}
+		FleshAsset->Dataflow->Schema = UDataflowSchema::StaticClass();
+		Dataflow = FleshAsset->Dataflow;
+
+		NodeDetailsEditor = CreateNodeDetailsEditorWidget(ObjectToEdit);
+		AssetDetailsEditor = CreateAssetDetailsEditorWidget(FleshAsset);
+		GraphEditor = CreateGraphEditorWidget(FleshAsset->Dataflow, NodeDetailsEditor);
+		SkeletalEditor = CreateSkeletalEditorWidget(FleshAsset->SkeletalMesh);
+
+		Context = TSharedPtr< Dataflow::FEngineContext>(new Dataflow::FEngineContext(FleshAsset, Dataflow, FPlatformTime::Cycles64(), FString("UFleshAsset")));
+		LastNodeTimestamp = Context->GetTimestamp();
+
+		const TSharedRef<FTabManager::FLayout> StandaloneDefaultLayout = FTabManager::NewLayout("FleshAsset_Layout.V1")
+			->AddArea
+			(
+				FTabManager::NewPrimaryArea()->SetOrientation(Orient_Vertical)
+				->Split
+				(
+					FTabManager::NewSplitter()->SetOrientation(Orient_Vertical)
+					->Split
+					(
+						FTabManager::NewSplitter()->SetOrientation(Orient_Horizontal)
+						->SetSizeCoefficient(0.9f)
+						->Split
+						(
+							FTabManager::NewStack()
+							->SetSizeCoefficient(0.6f)
+							->AddTab(GraphCanvasTabId, ETabState::OpenedTab)
+							->SetHideTabWell(true)
+						)
+						->Split
+						(
+							FTabManager::NewSplitter()->SetOrientation(Orient_Vertical)
+							->SetSizeCoefficient(0.2f)
+							->Split
+							(
+								FTabManager::NewStack()
+								->SetSizeCoefficient(0.7f)
+								->AddTab(NodeDetailsTabId, ETabState::OpenedTab)
+							)
+						)
+					)
+				)
+			);
+			
+		const bool bCreateDefaultStandaloneMenu = true;
+		const bool bCreateDefaultToolbar = true;
+		FAssetEditorToolkit::InitAssetEditor(Mode, InitToolkitHost, FName(TEXT("FleshEditorApp")), StandaloneDefaultLayout, bCreateDefaultToolbar, bCreateDefaultStandaloneMenu, ObjectToEdit);
+	}
+}
+
+void FFleshEditorToolkit::OnPropertyValueChanged(const FPropertyChangedEvent& PropertyChangedEvent)
+{
+	FDataflowEditorCommands::OnPropertyValueChanged(this->GetDataflow(), Context, LastNodeTimestamp, PropertyChangedEvent);
+}
+
+bool FFleshEditorToolkit::OnNodeVerifyTitleCommit(const FText& NewText, UEdGraphNode* GraphNode, FText& OutErrorMessage)
+{
+	return FDataflowEditorCommands::OnNodeVerifyTitleCommit(NewText, GraphNode, OutErrorMessage);
+}
+
+void FFleshEditorToolkit::OnNodeTitleCommitted(const FText& InNewText, ETextCommit::Type InCommitType, UEdGraphNode* GraphNode)
+{
+	FDataflowEditorCommands::OnNodeTitleCommitted(InNewText, InCommitType, GraphNode);
+}
+
+void FFleshEditorToolkit::Tick(float DeltaTime)
+{
+	if (bDataflowAssetEditorFleshLiveEvaluationEnableCVar)
+	{
+		if (Dataflow && FleshAsset)
+		{
+			if (!Context)
+			{
+				Context = TSharedPtr< Dataflow::FEngineContext>(new Dataflow::FEngineContext(FleshAsset, Dataflow, Dataflow::FTimestamp::Invalid, FString("UFleshAsset")));
+				LastNodeTimestamp = Dataflow::FTimestamp::Invalid;
+			}
+
+			FDataflowEditorCommands::EvaluateNode(*Context.Get(), LastNodeTimestamp, Dataflow, nullptr, nullptr, FleshAsset->Terminal);
+		}
+	}
+}
+
+TStatId FFleshEditorToolkit::GetStatId() const
+{
+	RETURN_QUICK_DECLARE_CYCLE_STAT(FFleshEditorToolkit, STATGROUP_Tickables);
+}
+
+
+TSharedRef<SGraphEditor> FFleshEditorToolkit::CreateGraphEditorWidget(UDataflow* DataflowToEdit, TSharedPtr<IStructureDetailsView> InNodeDetailsEditor)
+{
+	ensure(DataflowToEdit);
+	IDataflowEditorPlugin& DataflowEditorModule = FModuleManager::LoadModuleChecked<IDataflowEditorPlugin>(TEXT("DataflowEditor"));
+
+	FDataflowEditorCommands::FGraphEvaluationCallback Evaluate = [&](const FDataflowNode* Node, const FDataflowOutput* Out)
+	{
+		if (Dataflow && FleshAsset)
+		{
+			if (!Context)
+			{
+				Context = TSharedPtr< Dataflow::FEngineContext>(new Dataflow::FEngineContext(FleshAsset, Dataflow, Dataflow::FTimestamp::Invalid, FString("UFleshAsset")));
+			}
+			LastNodeTimestamp = Dataflow::FTimestamp::Invalid;
+
+			FDataflowEditorCommands::EvaluateNode(*Context.Get(),LastNodeTimestamp,Dataflow,Node,Out);
+		}
+	};
+
+	FDataflowEditorCommands::FOnDragDropEventCallback  DragDropEvent = [&](const FGeometry& MyGeometry, const FDragDropEvent& DragDropEvent)
+	{
+		TSharedPtr<FAssetSchemaAction_Dataflow_CreateNode_DataflowEdNode> Action = 
+			FAssetSchemaAction_Dataflow_CreateNode_DataflowEdNode::CreateAction(Dataflow, "SkeletalMeshBone");
+		Action->PerformAction(Dataflow, nullptr, FVector2D(0,0), true);
+	};
+
+	SGraphEditor::FGraphEditorEvents InEvents;
+	InEvents.OnVerifyTextCommit = FOnNodeVerifyTextCommit::CreateSP(this, &FFleshEditorToolkit::OnNodeVerifyTitleCommit);
+	InEvents.OnTextCommitted = FOnNodeTextCommitted::CreateSP(this, &FFleshEditorToolkit::OnNodeTitleCommitted);
+
+	UE_LOG(LogChaosFleshEditorToolkit, All, TEXT("FFleshEditorToolkit::CreateGraphEditorWidget"));
+	return SNew(SDataflowGraphEditor, FleshAsset)
+		.GraphToEdit(DataflowToEdit)
+		.GraphEvents(InEvents)
+		.DetailsView(InNodeDetailsEditor)
+		.EvaluateGraph(Evaluate)
+		.OnDragDropEvent(DragDropEvent);
+}
+
+TSharedPtr<IStructureDetailsView> FFleshEditorToolkit::CreateNodeDetailsEditorWidget(UObject* ObjectToEdit)
+{
+	ensure(ObjectToEdit);
+	FPropertyEditorModule& PropertyEditorModule = FModuleManager::LoadModuleChecked<FPropertyEditorModule>(TEXT("PropertyEditor"));
+
+	FDetailsViewArgs DetailsViewArgs;
+	{
+		DetailsViewArgs.bAllowSearch = false;
+		DetailsViewArgs.bHideSelectionTip = true;
+		DetailsViewArgs.bLockable = false;
+		DetailsViewArgs.bSearchInitialKeyFocus = true;
+		DetailsViewArgs.bUpdatesFromSelection = false;
+		DetailsViewArgs.NotifyHook = nullptr;
+		DetailsViewArgs.bShowOptions = true;
+		DetailsViewArgs.bShowModifiedPropertiesOption = false;
+		DetailsViewArgs.bShowScrollBar = false;
+	}
+
+	FStructureDetailsViewArgs StructureViewArgs;
+	{
+		StructureViewArgs.bShowObjects = true;
+		StructureViewArgs.bShowAssets = true;
+		StructureViewArgs.bShowClasses = true;
+		StructureViewArgs.bShowInterfaces = true;
+	}
+
+	TSharedPtr<IStructureDetailsView> DetailsView = PropertyEditorModule.CreateStructureDetailView(DetailsViewArgs, StructureViewArgs, nullptr);
+	DetailsView->GetDetailsView()->SetObject(ObjectToEdit);
+	DetailsView->GetOnFinishedChangingPropertiesDelegate().AddSP(this, &FFleshEditorToolkit::OnPropertyValueChanged);
+
+	return DetailsView;
+
+}
+
+
+TSharedPtr<IDetailsView> FFleshEditorToolkit::CreateAssetDetailsEditorWidget(UObject* ObjectToEdit)
+{
+	ensure(ObjectToEdit);
+	FPropertyEditorModule& PropertyEditorModule = FModuleManager::LoadModuleChecked<FPropertyEditorModule>(TEXT("PropertyEditor"));
+
+	FDetailsViewArgs DetailsViewArgs;
+	{
+		DetailsViewArgs.bAllowSearch = true;
+		DetailsViewArgs.bLockable = false;
+		DetailsViewArgs.bUpdatesFromSelection = false;
+		DetailsViewArgs.NameAreaSettings = FDetailsViewArgs::HideNameArea;
+		DetailsViewArgs.NotifyHook = this;
+	}
+
+	TSharedPtr<IDetailsView> DetailsView = PropertyEditorModule.CreateDetailView(DetailsViewArgs);
+	DetailsView->SetObject(ObjectToEdit);
+	return DetailsView;
+}
+
+
+TSharedPtr<ISkeletonTree> FFleshEditorToolkit::CreateSkeletalEditorWidget(USkeletalMesh* ObjectToEdit)
+{
+	check(FleshAsset);
+
+	if (!StubSkeletalMesh)
+	{
+		const FName NodeName = MakeUniqueObjectName(FleshAsset, UDataflow::StaticClass(), FName("USkeleton"));
+		StubSkeleton = NewObject<USkeleton>(FleshAsset, NodeName);
+		const FName NodeName2 = MakeUniqueObjectName(FleshAsset, UDataflow::StaticClass(), FName("USkeleton"));
+		StubSkeletalMesh = NewObject<USkeletalMesh>(FleshAsset, NodeName2);
+		StubSkeletalMesh->SetSkeleton(StubSkeleton);
+	}
+
+	FSkeletonTreeArgs SkeletonTreeArgs;
+	//SkeletonTreeArgs.OnSelectionChanged = FOnSkeletonTreeSelectionChanged::CreateSP(this, &FAnimationBlueprintEditor::HandleSelectionChanged);
+	//SkeletonTreeArgs.PreviewScene = GetPreviewScene();
+	//SkeletonTreeArgs.ContextName = GetToolkitFName();
+
+	USkeleton* Skeleton = StubSkeleton;
+	if (FleshAsset && FleshAsset->SkeletalMesh && FleshAsset->SkeletalMesh->GetSkeleton())
+		Skeleton = FleshAsset->SkeletalMesh->GetSkeleton();
+
+	ISkeletonEditorModule& SkeletonEditorModule = FModuleManager::LoadModuleChecked<ISkeletonEditorModule>("SkeletonEditor");
+	TSharedPtr<ISkeletonTree> SkeletonTree = SkeletonEditorModule.CreateSkeletonTree(Skeleton, SkeletonTreeArgs);
+	//AddApplicationMode(AnimationEditorModes::AnimationEditorMode,MakeShareable(new FAnimationEditorMode(SharedThis(this), SkeletonTree.ToSharedRef())));
+	return SkeletonTree;
+}
+
+TSharedRef<SDockTab> FFleshEditorToolkit::SpawnTab_Viewport(const FSpawnTabArgs& Args)
+{
+	check(Args.GetTabId() == ViewportTabId);
+
+	TSharedRef< SDockTab > DockableTab = SNew(SDockTab);
+	ViewportEditor = MakeShareable(new FEditorViewportTabContent());
+	TWeakPtr<FFleshEditorToolkit> WeakSharedThis = SharedThis(this);
+
+	const FString LayoutId = FString("FleshEditorViewport");
+	ViewportEditor->Initialize([WeakSharedThis](const FAssetEditorViewportConstructionArgs& InConstructionArgs)
+		{
+			return SNew(SFleshEditorViewport)
+				.FleshEditorToolkit(WeakSharedThis);
+		}, DockableTab, LayoutId);
+
+	return DockableTab;
+}
+
+TSharedRef<SDockTab> FFleshEditorToolkit::SpawnTab_GraphCanvas(const FSpawnTabArgs& Args)
+{
+	check(Args.GetTabId() == GraphCanvasTabId);
+
+	TSharedRef<SDockTab> SpawnedTab = SNew(SDockTab)
+		.Label(LOCTEXT("FleshDataflowEditor_Dataflow_TabTitle", "Graph"))
+		[
+			GraphEditor.ToSharedRef()
+		];
+
+	return SpawnedTab;
+}
+
+TSharedRef<SDockTab> FFleshEditorToolkit::SpawnTab_NodeDetails(const FSpawnTabArgs& Args)
+{
+	check(Args.GetTabId() == NodeDetailsTabId);
+
+	return SNew(SDockTab)
+		.Label(LOCTEXT("FleshEditorNodeDetails_TabTitle", "Details"))
+		[
+			NodeDetailsEditor->GetWidget()->AsShared()
+		];
+}
+
+TSharedRef<SDockTab> FFleshEditorToolkit::SpawnTab_AssetDetails(const FSpawnTabArgs& Args)
+{
+	check(Args.GetTabId() == AssetDetailsTabId);
+
+	return SNew(SDockTab)
+		.Label(LOCTEXT("FleshEditorAssetDetails_TabTitle", "Asset Details"))
+		[
+			AssetDetailsEditor.ToSharedRef()
+		];
+}
+
+TSharedRef<SDockTab> FFleshEditorToolkit::SpawnTab_Skeletal(const FSpawnTabArgs& Args)
+{
+	check(Args.GetTabId() == SkeletalTabId);
+
+	USkeletalMesh* SkeletalMesh = StubSkeletalMesh;
+	if (FleshAsset && FleshAsset->SkeletalMesh)
+		SkeletalMesh = FleshAsset->SkeletalMesh;
+
+	SkeletalEditor->SetSkeletalMesh(SkeletalMesh);
+
+	return SNew(SDockTab)
+		.Label(LOCTEXT("FleshEditorSkeletal_TabTitle", "Skeletal Hierarchy"))
+		[
+			SkeletalEditor.ToSharedRef()
+		];
+}
+
+
+void FFleshEditorToolkit::RegisterTabSpawners(const TSharedRef<FTabManager>& InTabManager)
+{
+	TSharedRef<FWorkspaceItem> WorkspaceMenuCategoryRef = InTabManager->AddLocalWorkspaceMenuCategory(LOCTEXT("WorkspaceMenu_FleshEditorEditor", "Flesh Editor"));
+
+	FAssetEditorToolkit::RegisterTabSpawners(InTabManager);
+	InTabManager->RegisterTabSpawner(ViewportTabId, FOnSpawnTab::CreateSP(this, &FFleshEditorToolkit::SpawnTab_Viewport))
+		.SetDisplayName(LOCTEXT("FleshViewportTab", "Flesh Viewport"))
+		.SetGroup(WorkspaceMenuCategoryRef)
+		.SetIcon(FSlateIcon(FAppStyle::GetAppStyleSetName(), "GraphEditor.EventGraph_16x"));
+
+	InTabManager->RegisterTabSpawner(GraphCanvasTabId, FOnSpawnTab::CreateSP(this, &FFleshEditorToolkit::SpawnTab_GraphCanvas))
+		.SetDisplayName(LOCTEXT("FleshDataflowTab", "Graph"))
+		.SetGroup(WorkspaceMenuCategoryRef)
+		.SetIcon(FSlateIcon(FAppStyle::GetAppStyleSetName(), "GraphEditor.EventGraph_16x"));
+
+	InTabManager->RegisterTabSpawner(NodeDetailsTabId, FOnSpawnTab::CreateSP(this, &FFleshEditorToolkit::SpawnTab_NodeDetails))
+		.SetDisplayName(LOCTEXT("FleshNodeDetailsTab", "Details"))
+		.SetGroup(WorkspaceMenuCategoryRef)
+		.SetIcon(FSlateIcon(FAppStyle::GetAppStyleSetName(), "LevelEditor.Tabs.Details"));
+
+	InTabManager->RegisterTabSpawner(AssetDetailsTabId, FOnSpawnTab::CreateSP(this, &FFleshEditorToolkit::SpawnTab_AssetDetails))
+		.SetDisplayName(LOCTEXT("FleshAssetDetailsTab", "Asset Details"))
+		.SetGroup(WorkspaceMenuCategoryRef)
+		.SetIcon(FSlateIcon(FAppStyle::GetAppStyleSetName(), "LevelEditor.Tabs.AssetDetails"));
+
+	InTabManager->RegisterTabSpawner(SkeletalTabId, FOnSpawnTab::CreateSP(this, &FFleshEditorToolkit::SpawnTab_Skeletal))
+		.SetDisplayName(LOCTEXT("FleshSkeletalTab", "Skeletal Hierarchy"))
+		.SetGroup(WorkspaceMenuCategoryRef)
+		.SetIcon(FSlateIcon(FAppStyle::GetAppStyleSetName(), "LevelEditor.Tabs.SkeletalHierarchy"));
+}
+
+
+void FFleshEditorToolkit::GetSaveableObjects(TArray<UObject*>& OutObjects) const
+{
+	if(FleshAsset)
+		OutObjects.Add(static_cast<UObject*>(FleshAsset));
+	if (Dataflow)
+		OutObjects.Add(static_cast<UObject*>(Dataflow));
+}
+
+
+FName FFleshEditorToolkit::GetToolkitFName() const
+{
+	return FName("FleshEditor");
+}
+
+FText FFleshEditorToolkit::GetBaseToolkitName() const
+{
+	return LOCTEXT("AppLabel", "Flesh Editor");
+}
+
+FString FFleshEditorToolkit::GetWorldCentricTabPrefix() const
+{
+	return LOCTEXT("WorldCentricTabPrefix", "Flesh").ToString();
+}
+
+FLinearColor FFleshEditorToolkit::GetWorldCentricTabColorScale() const
+{
+	return FLinearColor(0.3f, 0.2f, 0.5f, 0.5f);
+}
+
+FString FFleshEditorToolkit::GetReferencerName() const
+{
+	return TEXT("FFleshEditorToolkit");
+}
+
+void FFleshEditorToolkit::AddReferencedObjects(FReferenceCollector& Collector)
+{
+	Collector.AddReferencedObject(FleshAsset);
+}
+#undef LOCTEXT_NAMESPACE
diff --git a/Engine/Plugins/Experimental/ChaosFlesh/Source/ChaosFleshEditor/Private/ChaosFlesh/Editor/FleshEditorToolkit.h b/Engine/Plugins/Experimental/ChaosFlesh/Source/ChaosFleshEditor/Private/ChaosFlesh/Editor/FleshEditorToolkit.h
new file mode 100644
--- /dev/null
+++ b/Engine/Plugins/Experimental/ChaosFlesh/Source/ChaosFleshEditor/Private/ChaosFlesh/Editor/FleshEditorToolkit.h
@@ -0,0 +1,107 @@
+// Copyright Epic Games, Inc. All Rights Reserved.
+
+#pragma once
+
+#include "CoreMinimal.h"
+#include "Dataflow/DataflowObjectInterface.h"
+#include "GraphEditor.h"
+#include "ISkeletonTree.h"
+#include "Misc/NotifyHook.h"
+#include "TickableEditorObject.h"
+#include "Toolkits/SimpleAssetEditor.h"
+
+class FEditorViewportTabContent;
+class IStructureDetailsView;
+class ISkeletalView;
+class FTabManager;
+class IToolkitHost;
+class UFleshAsset;
+class UDataflow;
+class USkeletalMesh;
+
+class FFleshEditorToolkit : public FAssetEditorToolkit, public FTickableEditorObject, public FNotifyHook, public FGCObject
+{
+public:
+	void InitFleshAssetEditor(const EToolkitMode::Type Mode, const TSharedPtr<IToolkitHost>& InitToolkitHost, UObject* ObjectToEdit);
+
+	//~ Begin DataflowEditorActions
+	void OnPropertyValueChanged(const FPropertyChangedEvent& PropertyChangedEvent);
+	bool OnNodeVerifyTitleCommit(const FText& NewText, UEdGraphNode* GraphNode, FText& OutErrorMessage);
+	void OnNodeTitleCommitted(const FText& InNewText, ETextCommit::Type InCommitType, UEdGraphNode* GraphNode);
+	//~ End DataflowEditorActions
+
+
+	// IToolkit Interface
+	virtual FName GetToolkitFName() const override;
+	virtual FText GetBaseToolkitName() const override;
+	virtual FString GetWorldCentricTabPrefix() const override;
+	virtual FLinearColor GetWorldCentricTabColorScale() const override;
+
+	// FGCObject Interface
+	virtual void AddReferencedObjects(FReferenceCollector& Collector) override;
+	virtual FString GetReferencerName() const override;
+
+	// Tab spawners 
+	virtual void RegisterTabSpawners(const TSharedRef<FTabManager>& TabManager) override;
+	TSharedRef<SDockTab> SpawnTab_Viewport(const FSpawnTabArgs& Args);
+	TSharedRef<SDockTab> SpawnTab_GraphCanvas(const FSpawnTabArgs& Args);
+	TSharedRef<SDockTab> SpawnTab_NodeDetails(const FSpawnTabArgs& Args);
+	TSharedRef<SDockTab> SpawnTab_AssetDetails(const FSpawnTabArgs& Args);
+	TSharedRef<SDockTab> SpawnTab_Skeletal(const FSpawnTabArgs& Args);
+
+	UFleshAsset* GetFleshAsset() {return FleshAsset;}
+	const UFleshAsset* GetFleshAsset() const { return FleshAsset; }
+
+	UDataflow* GetDataflow() { return Dataflow; }
+	const UDataflow* GetDataflow() const { return Dataflow; }
+
+	TSharedPtr<Dataflow::FEngineContext> GetContext() { return Context; }
+	TSharedPtr<const Dataflow::FEngineContext> GetContext() const { return Context; }
+
+	//~ Begin FTickableEditorObject interface
+	virtual void Tick(float DeltaTime) override;
+	virtual bool IsTickable() const override { return true; }
+	virtual TStatId GetStatId() const override;
+	//~ End FTickableEditorObject Interface
+
+
+
+	//~ Begin FAssetEditorToolkit interface
+	virtual void GetSaveableObjects(TArray<UObject*>& OutObjects) const override;
+	//virtual void SaveAsset_Execute() override;
+	//virtual void SaveAssetAs_Execute() override;
+	//virtual bool OnRequestClose() override;
+	//~ End FAssetEditorToolkit interface
+
+
+private:
+	static const FName ViewportTabId;
+	TSharedPtr<FEditorViewportTabContent> ViewportEditor;
+
+	static const FName GraphCanvasTabId;
+	TSharedPtr<SGraphEditor> GraphEditor;
+	TSharedPtr<FUICommandList> GraphEditorCommands;
+	TSharedRef<SGraphEditor> CreateGraphEditorWidget(UDataflow* ObjectToEdit, TSharedPtr<IStructureDetailsView> InNodeDetailsEditor);
+
+	static const FName NodeDetailsTabId;
+	TSharedPtr<IStructureDetailsView> NodeDetailsEditor;
+	TSharedPtr<IStructureDetailsView> CreateNodeDetailsEditorWidget(UObject* ObjectToEdit);
+
+	static const FName AssetDetailsTabId;
+	TSharedPtr<IDetailsView> AssetDetailsEditor;
+	TSharedPtr<IDetailsView> CreateAssetDetailsEditorWidget(UObject* ObjectToEdit);
+
+
+	static const FName SkeletalTabId;
+	TObjectPtr<USkeleton> StubSkeleton;
+	TObjectPtr<USkeletalMesh> StubSkeletalMesh;
+	TSharedPtr<class ISkeletonTree> SkeletalEditor;
+	TSharedPtr<ISkeletonTree> CreateSkeletalEditorWidget(USkeletalMesh* ObjectToEdit);
+
+	UFleshAsset* FleshAsset = nullptr;
+	UDataflow* Dataflow = nullptr;
+
+	TSharedPtr<Dataflow::FEngineContext> Context;
+	Dataflow::FTimestamp LastNodeTimestamp = Dataflow::FTimestamp::Invalid;
+
+};
\ No newline at end of file
diff --git a/Engine/Plugins/Experimental/ChaosFlesh/Source/ChaosFleshEditor/Private/ChaosFlesh/Editor/FleshEditorViewport.cpp b/Engine/Plugins/Experimental/ChaosFlesh/Source/ChaosFleshEditor/Private/ChaosFlesh/Editor/FleshEditorViewport.cpp
new file mode 100644
--- /dev/null
+++ b/Engine/Plugins/Experimental/ChaosFlesh/Source/ChaosFleshEditor/Private/ChaosFlesh/Editor/FleshEditorViewport.cpp
@@ -0,0 +1,165 @@
+// Copyright Epic Games, Inc. All Rights Reserved.
+#include "FleshEditorViewport.h"
+
+#include "AdvancedPreviewScene.h"
+#include "ChaosFlesh/FleshActor.h"
+#include "ChaosFlesh/FleshAsset.h"
+#include "Dataflow/DataflowEdNode.h"
+#include "Dataflow/DataflowObject.h"
+#include "Dataflow/DataflowRenderingActor.h"
+#include "Dataflow/DataflowRenderingComponent.h"
+#include "EditorViewportClient.h"
+#include "FleshEditorToolkit.h"
+#include "SCommonEditorViewportToolbarBase.h"
+
+SFleshEditorViewport::SFleshEditorViewport()
+{
+	// Temporarily allow water subsystem to be created on preview worlds because we need one here : 
+	//UWaterSubsystem::FScopedAllowWaterSubsystemOnPreviewWorld AllowWaterSubsystemOnPreviewWorldScope(true);
+	PreviewScene = MakeShareable(new FAdvancedPreviewScene(FPreviewScene::ConstructionValues()));
+	PreviewScene->SetFloorVisibility(false);
+}
+
+void SFleshEditorViewport::Construct(const FArguments& InArgs)
+{
+	FleshEditorToolkitPtr = InArgs._FleshEditorToolkit;
+	TSharedPtr<FFleshEditorToolkit> FleshEditorToolkit = FleshEditorToolkitPtr.Pin();
+	check(FleshEditorToolkitPtr.IsValid());
+
+	SEditorViewport::Construct(SEditorViewport::FArguments());
+
+	FBoxSphereBounds SphereBounds1 = FBoxSphereBounds(EForceInit::ForceInitToZero);
+	CustomFleshActor =CastChecked<AFleshActor>(PreviewScene->GetWorld()->SpawnActor(AFleshActor::StaticClass()));
+	if (UFleshAsset* FleshAsset = FleshEditorToolkit->GetFleshAsset())
+	{
+		if (UFleshComponent* FleshComponent = CustomFleshActor->GetFleshComponent())
+		{
+			FleshComponent->SetRestCollection(FleshAsset);
+			SphereBounds1 = FleshComponent->CalcBounds(CustomFleshActor->GetTransform());
+		}
+	}
+
+	FBoxSphereBounds SphereBounds2 = FBoxSphereBounds(EForceInit::ForceInitToZero);
+	CustomDataflowRenderingActor = CastChecked<ADataflowRenderingActor>(PreviewScene->GetWorld()->SpawnActor(ADataflowRenderingActor::StaticClass()));
+
+	EditorViewportClient->SetDataflowRenderingActor(CustomDataflowRenderingActor);
+	EditorViewportClient->FocusViewportOnBox( (/*SphereBounds1 + */ SphereBounds2).GetBox());
+}
+
+TSharedRef<SEditorViewport> SFleshEditorViewport::GetViewportWidget()
+{
+	return SharedThis(this);
+}
+
+TSharedPtr<FExtender> SFleshEditorViewport::GetExtenders() const
+{
+	TSharedPtr<FExtender> Result(MakeShareable(new FExtender));
+	return Result;
+}
+
+void SFleshEditorViewport::OnFloatingButtonClicked()
+{
+}
+
+void SFleshEditorViewport::AddReferencedObjects(FReferenceCollector& Collector)
+{							
+	Collector.AddReferencedObject(CustomFleshActor);
+	Collector.AddReferencedObject(CustomDataflowRenderingActor);
+}
+
+TSharedRef<FEditorViewportClient> SFleshEditorViewport::MakeEditorViewportClient()
+{
+	EditorViewportClient = MakeShareable(new FFleshEditorViewportClient(PreviewScene.Get(), SharedThis(this), FleshEditorToolkitPtr));
+	return EditorViewportClient.ToSharedRef();
+}
+
+TSharedPtr<SWidget> SFleshEditorViewport::MakeViewportToolbar()
+{
+	return SNew(SCommonEditorViewportToolbarBase, SharedThis(this));
+}
+
+
+// ----------------------------------------------------------------------------------
+
+FFleshEditorViewportClient::FFleshEditorViewportClient(FPreviewScene* InPreviewScene, 
+	const TWeakPtr<SEditorViewport> InEditorViewportWidget,
+	TWeakPtr<FFleshEditorToolkit> InFleshEditorToolkitPtr)
+	: 
+	FEditorViewportClient(nullptr, InPreviewScene, InEditorViewportWidget)
+	, FleshEditorToolkitPtr(InFleshEditorToolkitPtr)
+{
+	bSetListenerPosition = false;
+	SetRealtime(true);
+	EngineShowFlags.Grid = false;
+}
+
+Dataflow::FTimestamp FFleshEditorViewportClient::LatestTimestamp(const UDataflow* Dataflow, const Dataflow::FContext* Context)
+{
+	if (Dataflow && Context)
+	{
+		return FMath::Max(Dataflow->GetRenderingTimestamp().Value, Context->GetTimestamp().Value);
+	}
+	return Dataflow::FTimestamp::Invalid;
+}
+
+
+void FFleshEditorViewportClient::Tick(float DeltaSeconds)
+{
+	Super::Tick(DeltaSeconds);
+
+	TSharedPtr<FFleshEditorToolkit> FleshEditorToolkit = FleshEditorToolkitPtr.Pin();
+
+	if (DataflowRenderingActor && FleshEditorToolkitPtr.IsValid())
+	{
+		if (TSharedPtr<Dataflow::FContext> Context = FleshEditorToolkit->GetContext())
+		{
+			if (UDataflowRenderingComponent* DataflowRenderingComponent = DataflowRenderingActor->GetDataflowRenderingComponent())
+			{
+				if (const UDataflow* Dataflow = FleshEditorToolkit->GetDataflow())
+				{
+					Dataflow::FTimestamp SystemTimestamp = LatestTimestamp(Dataflow, Context.Get());
+					if (SystemTimestamp >= LastModifiedTimestamp)
+					{
+						if (Dataflow->GetRenderTargets().Num())
+						{
+							// @todo(dataflow) : Check the Time on the target outs first instead.
+							//					 to avoid invalidation during unrelated edits. 
+
+							DataflowRenderingComponent->ResetRenderTargets();
+
+							DataflowRenderingComponent->SetDataflow(Dataflow);
+							DataflowRenderingComponent->SetContext(Context);
+							for (const UDataflowEdNode* Node : Dataflow->GetRenderTargets())
+							{
+								DataflowRenderingComponent->AddRenderTarget(Node);
+							}
+						}
+						else
+						{
+							DataflowRenderingComponent->ResetRenderTargets();
+						}
+						LastModifiedTimestamp = LatestTimestamp(Dataflow, Context.Get()).Value + 1;
+					}
+				}
+			}
+		}
+	}
+	// Tick the preview scene world.
+	if (!GIntraFrameDebuggingGameThread)
+	{
+		PreviewScene->GetWorld()->Tick(IsRealtime() ? LEVELTICK_All : LEVELTICK_TimeOnly, DeltaSeconds);
+	}
+}
+
+
+void FFleshEditorViewportClient::AddReferencedObjects(FReferenceCollector& Collector)
+{
+	Super::AddReferencedObjects(Collector);
+
+	if (DataflowRenderingActor)
+	{
+		Collector.AddReferencedObject(DataflowRenderingActor);
+	}
+}
+
+
diff --git a/Engine/Plugins/Experimental/ChaosFlesh/Source/ChaosFleshEditor/Private/ChaosFlesh/Editor/FleshEditorViewport.h b/Engine/Plugins/Experimental/ChaosFlesh/Source/ChaosFleshEditor/Private/ChaosFlesh/Editor/FleshEditorViewport.h
new file mode 100644
--- /dev/null
+++ b/Engine/Plugins/Experimental/ChaosFlesh/Source/ChaosFleshEditor/Private/ChaosFlesh/Editor/FleshEditorViewport.h
@@ -0,0 +1,83 @@
+// Copyright Epic Games, Inc. All Rights Reserved.
+#pragma once
+
+#include "CoreMinimal.h"
+#include "Dataflow/DataflowNodeParameters.h"
+#include "EditorViewportClient.h"
+#include "SAssetEditorViewport.h"
+#include "SCommonEditorViewportToolbarBase.h"
+
+class FFleshEditorToolkit;
+class FAdvancedPreviewScene;
+class FEditorViewportClient;
+class FFleshEditorViewportClient;
+class SEditorViewport;
+class AFleshActor;
+class ADataflowRenderingActor;
+
+// ----------------------------------------------------------------------------------
+
+class SFleshEditorViewport : public SAssetEditorViewport, public ICommonEditorViewportToolbarInfoProvider, public FGCObject
+{
+public:
+	SLATE_BEGIN_ARGS(SFleshEditorViewport) {}
+		SLATE_ARGUMENT(TWeakPtr<FFleshEditorToolkit>, FleshEditorToolkit)
+	SLATE_END_ARGS()
+
+	SFleshEditorViewport();
+
+	void Construct(const FArguments& InArgs);
+
+	//~ ICommonEditorViewportToolbarInfoProvider interface
+	virtual TSharedRef<SEditorViewport> GetViewportWidget() override;
+	virtual TSharedPtr<FExtender> GetExtenders() const override;
+	virtual void OnFloatingButtonClicked() override;
+
+	//~ FGCObject Interface
+	virtual void AddReferencedObjects(FReferenceCollector& Collector) override;
+	virtual FString GetReferencerName() const override{return TEXT("SFleshEditorViewport");}
+
+protected:
+	virtual TSharedRef<FEditorViewportClient> MakeEditorViewportClient() override;
+	virtual TSharedPtr<SWidget> MakeViewportToolbar() override;
+private:
+	/// The scene for this viewport. 
+	TSharedPtr<FAdvancedPreviewScene> PreviewScene;
+
+	/// Editor viewport client 
+	TSharedPtr<FFleshEditorViewportClient> EditorViewportClient;
+
+	TWeakPtr<FFleshEditorToolkit> FleshEditorToolkitPtr;
+
+	AFleshActor* CustomFleshActor = nullptr;
+	ADataflowRenderingActor* CustomDataflowRenderingActor = nullptr;
+};
+
+
+// ----------------------------------------------------------------------------------
+
+class FFleshEditorViewportClient : public FEditorViewportClient
+{
+public:
+	using Super = FEditorViewportClient;
+	
+	FFleshEditorViewportClient(FPreviewScene* InPreviewScene,
+		const TWeakPtr<SEditorViewport> InEditorViewportWidget = nullptr,
+		TWeakPtr<FFleshEditorToolkit> InFleshEditorToolkitPtr = nullptr);
+
+
+	void SetDataflowRenderingActor(ADataflowRenderingActor* InActor) { DataflowRenderingActor = InActor; }
+	Dataflow::FTimestamp LatestTimestamp(const UDataflow* Dataflow, const Dataflow::FContext* Context);
+
+	// FEditorViewportClient interface
+	virtual void Tick(float DeltaSeconds) override;
+	// End of FEditorViewportClient
+
+	//~ FGCObject Interface
+	virtual void AddReferencedObjects(FReferenceCollector& Collector) override;
+	virtual FString GetReferencerName() const override { return TEXT("FFleshEditorViewportClient"); }
+
+	TWeakPtr<FFleshEditorToolkit> FleshEditorToolkitPtr = nullptr;
+	ADataflowRenderingActor* DataflowRenderingActor = nullptr;
+	Dataflow::FTimestamp LastModifiedTimestamp = Dataflow::FTimestamp::Invalid;
+};
diff --git a/Engine/Plugins/Experimental/ChaosFlesh/Source/ChaosFleshEditor/Public/ChaosFlesh/Asset/AssetTypeActions_ChaosDeformableSolverAsset.h b/Engine/Plugins/Experimental/ChaosFlesh/Source/ChaosFleshEditor/Public/ChaosFlesh/Asset/AssetTypeActions_ChaosDeformableSolverAsset.h
new file mode 100644
--- /dev/null
+++ b/Engine/Plugins/Experimental/ChaosFlesh/Source/ChaosFleshEditor/Public/ChaosFlesh/Asset/AssetTypeActions_ChaosDeformableSolverAsset.h
@@ -0,0 +1,20 @@
+// Copyright Epic Games, Inc. All Rights Reserved.
+
+#pragma once
+
+#include "CoreMinimal.h"
+#include "Toolkits/IToolkitHost.h"
+#include "AssetTypeActions_Base.h"
+
+class FAssetTypeActions_ChaosDeformableSolver : public FAssetTypeActions_Base
+{
+public:
+	// IAssetTypeActions Implementation
+	virtual FText GetName() const override { return NSLOCTEXT("AssetTypeActions", "AssetTypeActions_ChaosDeformableSolver", "Chaos Deformable Solver"); }
+	virtual void GetActions(const TArray<UObject*>& InObjects, struct FToolMenuSection& Section) override;
+	virtual FColor GetTypeColor() const override { return FColor(255, 192, 128); }
+	virtual UClass* GetSupportedClass() const override;
+	virtual void OpenAssetEditor(const TArray<UObject*>& InObjects, TSharedPtr<class IToolkitHost> EditWithinLevelEditor = TSharedPtr<IToolkitHost>()) override;
+	virtual uint32 GetCategories() override { return EAssetTypeCategories::Physics; }
+	virtual UThumbnailInfo* GetThumbnailInfo(UObject* Asset) const override;
+};
diff --git a/Engine/Plugins/Experimental/ChaosFlesh/Source/ChaosFleshEditor/Public/ChaosFlesh/Asset/AssetTypeActions_FleshAsset.h b/Engine/Plugins/Experimental/ChaosFlesh/Source/ChaosFleshEditor/Public/ChaosFlesh/Asset/AssetTypeActions_FleshAsset.h
new file mode 100644
--- /dev/null
+++ b/Engine/Plugins/Experimental/ChaosFlesh/Source/ChaosFleshEditor/Public/ChaosFlesh/Asset/AssetTypeActions_FleshAsset.h
@@ -0,0 +1,21 @@
+// Copyright Epic Games, Inc. All Rights Reserved.
+
+#pragma once
+
+#include "CoreMinimal.h"
+#include "Toolkits/IToolkitHost.h"
+#include "AssetTypeActions_Base.h"
+
+class FAssetTypeActions_FleshAsset : public FAssetTypeActions_Base
+{
+public:
+	// IAssetTypeActions Implementation
+	virtual FText GetName() const override { return NSLOCTEXT("AssetTypeActions", "AssetTypeActions_FleshAsset", "Flesh Asset"); }
+	virtual void GetActions(const TArray<UObject*>& InObjects, struct FToolMenuSection& Section) override;
+	virtual FColor GetTypeColor() const override { return FColor(255, 192, 128); }
+	virtual UClass* GetSupportedClass() const override;
+	virtual uint32 GetCategories() override { return EAssetTypeCategories::Physics; }
+	virtual UThumbnailInfo* GetThumbnailInfo(UObject* Asset) const override;
+	virtual void OpenAssetEditor(const TArray<UObject*>& InObjects, TSharedPtr<class IToolkitHost> EditWithinLevelEditor = TSharedPtr<IToolkitHost>()) override;
+
+};
diff --git a/Engine/Plugins/Experimental/ChaosFlesh/Source/ChaosFleshEditor/Public/ChaosFlesh/Asset/ChaosAssetDeformableSolverFactory.h b/Engine/Plugins/Experimental/ChaosFlesh/Source/ChaosFleshEditor/Public/ChaosFlesh/Asset/ChaosAssetDeformableSolverFactory.h
new file mode 100644
--- /dev/null
+++ b/Engine/Plugins/Experimental/ChaosFlesh/Source/ChaosFleshEditor/Public/ChaosFlesh/Asset/ChaosAssetDeformableSolverFactory.h
@@ -0,0 +1,29 @@
+// Copyright Epic Games, Inc. All Rights Reserved.
+
+/** Factory which allows import of an ChaosSolverAsset */
+
+#pragma once
+
+#include "CoreMinimal.h"
+#include "UObject/ObjectMacros.h"
+#include "Factories/Factory.h"
+#include "ChaosFlesh/ChaosDeformableSolverAsset.h"
+
+#include "ChaosAssetDeformableSolverFactory.generated.h"
+
+
+/**
+* Factory for Simple Cube
+*/
+
+UCLASS()
+class CHAOSFLESHEDITOR_API UChaosDeformableSolverFactory : public UFactory
+{
+    GENERATED_UCLASS_BODY()
+
+	//~ Begin UFactory Interface
+	virtual UObject* FactoryCreateNew(UClass* Class, UObject* InParent, FName Name, EObjectFlags Flags, UObject* Context, FFeedbackContext* Warn) override;
+	static UChaosDeformableSolver* StaticFactoryCreateNew(UClass* Class, UObject* InParent, FName Name, EObjectFlags Flags, UObject* Context, FFeedbackContext* Warn);
+};
+
+
diff --git a/Engine/Plugins/Experimental/ChaosFlesh/Source/ChaosFleshEngine/Private/ChaosCache/FleshComponentCacheAdapter.cpp b/Engine/Plugins/Experimental/ChaosFlesh/Source/ChaosFleshEngine/Private/ChaosCache/FleshComponentCacheAdapter.cpp
new file mode 100644
--- /dev/null
+++ b/Engine/Plugins/Experimental/ChaosFlesh/Source/ChaosFleshEngine/Private/ChaosCache/FleshComponentCacheAdapter.cpp
@@ -0,0 +1,308 @@
+﻿// Copyright Epic Games, Inc. All Rights Reserved.
+
+#include "ChaosCache/FleshComponentCacheAdapter.h"
+
+#include "Chaos/Deformable/ChaosDeformableSolver.h"
+#include "Chaos/Deformable/ChaosDeformableSolverProxy.h"
+#include "Chaos/Deformable/ChaosDeformableSolverTypes.h"
+
+#include "Chaos/Adapters/CacheAdapter.h"
+#include "Chaos/ChaosCache.h"
+#include "Chaos/CacheManagerActor.h"
+#include "Chaos/ParticleHandle.h"
+#include "Chaos/PBDEvolution.h"
+#include "Chaos/PBDSoftsSolverParticles.h"
+#include "ChaosFlesh/ChaosDeformableSolverComponent.h"
+#include "ChaosFlesh/FleshComponent.h"
+
+
+namespace Chaos
+{
+	FComponentCacheAdapter::SupportType FFleshCacheAdapter::SupportsComponentClass(UClass* InComponentClass) const
+	{
+		const UClass* Desired = GetDesiredClass();
+		if(InComponentClass == Desired)
+		{
+			return FComponentCacheAdapter::SupportType::Direct;
+		}
+		else if(InComponentClass->IsChildOf(Desired))
+		{
+			return FComponentCacheAdapter::SupportType::Derived;
+		}
+
+		return FComponentCacheAdapter::SupportType::None;
+	}
+
+	UClass* FFleshCacheAdapter::GetDesiredClass() const
+	{
+		return UFleshComponent::StaticClass();
+	}
+
+	uint8 FFleshCacheAdapter::GetPriority() const
+	{
+		return EngineAdapterPriorityBegin;
+	}
+
+	void FFleshCacheAdapter::Record_PostSolve(UPrimitiveComponent* InComponent, const FTransform& InRootTransform, FPendingFrameWrite& OutFrame, Chaos::FReal InTime) const
+	{
+		if( FDeformableSolver* Solver = GetDeformableSolver(InComponent))
+		{
+			FDeformableSolver::FPhysicsThreadAccess PhysicsThreadAccess(Solver, Chaos::Softs::FPhysicsThreadAccessor());
+
+			if (FEvolution* Evolution = PhysicsThreadAccess.GetEvolution())
+			{
+				const FParticles& Particles = Evolution->Particles();
+
+				const uint32 NumParticles = Particles.Size();
+				if (NumParticles > 0)
+				{
+					const Softs::FSolverVec3* ParticleXs = &Particles.X(0);
+					const Softs::FSolverVec3* ParticleVs = &Particles.V(0);
+
+					for (uint32 ParticleIndex = 0; ParticleIndex < NumParticles; ++ParticleIndex)
+					{
+						const Softs::FSolverVec3& ParticleV = ParticleVs[ParticleIndex];
+						const Softs::FSolverVec3& ParticleX = ParticleXs[ParticleIndex];
+
+						// Adding the vertices relative position to the particle write datas
+						FPendingParticleWrite NewData;
+
+						NewData.ParticleIndex = ParticleIndex;
+						//	NewData.PendingTransform = FTransform(FQuat::Identity, FVector(ParticleX)).GetRelativeTransform(InRootTransform);
+
+						NewData.PendingCurveData.Add(MakeTuple(VelocityXName, ParticleV.X));
+						NewData.PendingCurveData.Add(MakeTuple(VelocityYName, ParticleV.Y));
+						NewData.PendingCurveData.Add(MakeTuple(VelocityZName, ParticleV.Z));
+
+						NewData.PendingCurveData.Add(MakeTuple(PositionXName, ParticleX.X));
+						NewData.PendingCurveData.Add(MakeTuple(PositionYName, ParticleX.Y));
+						NewData.PendingCurveData.Add(MakeTuple(PositionZName, ParticleX.Z));
+
+						OutFrame.PendingParticleData.Add(MoveTemp(NewData));
+					}
+				}
+			}
+		}
+	}
+
+	void FFleshCacheAdapter::Playback_PreSolve(UPrimitiveComponent* InComponent, UChaosCache* InCache, Chaos::FReal InTime, FPlaybackTickRecord& TickRecord, TArray<TPBDRigidParticleHandle<Chaos::FReal, 3>*>& OutUpdatedRigids) const
+	{
+		if (FDeformableSolver* Solver = GetDeformableSolver(InComponent))
+		{
+			FDeformableSolver::FPhysicsThreadAccess PhysicsThreadAccess(Solver, Softs::FPhysicsThreadAccessor());
+
+			if (FEvolution* Evolution = PhysicsThreadAccess.GetEvolution())
+			{
+				FCacheEvaluationContext Context(TickRecord);
+				Context.bEvaluateTransform = false;
+				Context.bEvaluateCurves = true;
+				Context.bEvaluateEvents = false;
+
+				// The evaluated result are already in world space since we are passing the tickrecord.spacetransform in the context
+				FCacheEvaluationResult EvaluatedResult = InCache->Evaluate(Context, nullptr);
+				const int32 NumCurves = EvaluatedResult.Curves.Num();
+
+				FParticles& Particles = Evolution->Particles();
+
+				const uint32 NumParticles = Particles.Size();
+				if (NumParticles > 0)
+				{
+					Softs::FSolverVec3* ParticleXs = &Particles.X(0);
+					Softs::FSolverVec3* ParticleVs = &Particles.V(0);
+
+					for (uint32 ParticleIndex = 0; ParticleIndex < NumParticles; ++ParticleIndex)
+					{
+						Softs::FSolverVec3& ParticleV = ParticleVs[ParticleIndex];
+						Softs::FSolverVec3& ParticleX = ParticleXs[ParticleIndex];
+
+						// Directly set the result of the cache into the solver particles
+						if ( ParticleIndex < (uint32)EvaluatedResult.ParticleIndices.Num())
+						{
+							const int32 TransformIndex = EvaluatedResult.ParticleIndices[ParticleIndex];
+							if (EvaluatedResult.Curves.IsValidIndex(TransformIndex))
+							{
+								if (const float* VelocityX = EvaluatedResult.Curves[TransformIndex].Find(VelocityXName))
+								{
+									ParticleV.X = *VelocityX;
+								}
+								if (const float* VelocityY = EvaluatedResult.Curves[TransformIndex].Find(VelocityYName))
+								{
+									ParticleV.Y = *VelocityY;
+								}
+								if (const float* VelocityZ = EvaluatedResult.Curves[TransformIndex].Find(VelocityZName))
+								{
+									ParticleV.Z = *VelocityZ;
+								}
+								if (const float* PositionX = EvaluatedResult.Curves[TransformIndex].Find(PositionXName))
+								{
+									ParticleX.X = *PositionX;
+								}
+								if (const float* PositionY = EvaluatedResult.Curves[TransformIndex].Find(PositionYName))
+								{
+									ParticleX.Y = *PositionY;
+								}
+								if (const float* PositionZ = EvaluatedResult.Curves[TransformIndex].Find(PositionZName))
+								{
+									ParticleX.Z = *PositionZ;
+								}
+							}
+						}
+					}
+				}
+			}
+		}
+	}
+
+	FGuid FFleshCacheAdapter::GetGuid() const
+	{
+		FGuid NewGuid;
+		checkSlow(FGuid::Parse(TEXT("2C054706CB7441B582377B0EDACD12EE"), NewGuid));
+		return NewGuid;
+	}
+
+	bool FFleshCacheAdapter::ValidForPlayback(UPrimitiveComponent* InComponent, UChaosCache* InCache) const
+	{
+		// If we have a flesh mesh we can play back any cache as long as it has one or more tracks
+		const UFleshComponent* FleshComp = CastChecked<UFleshComponent>(InComponent);
+		return FleshComp && InCache->TrackToParticle.Num() > 0;
+	}
+
+	Chaos::Softs::FDeformableSolver* FFleshCacheAdapter::GetDeformableSolver(UPrimitiveComponent* InComponent) const
+	{
+		if(InComponent)
+		{
+			if(UFleshComponent* FleshComp = CastChecked<UFleshComponent>(InComponent) )
+			{
+				if(FleshComp->GetDeformableSolver())
+				{
+					return FleshComp->GetDeformableSolver()->Solver.Get();
+				}
+			}
+		}
+		return nullptr;
+	}
+
+	Chaos::FPhysicsSolver* FFleshCacheAdapter::GetComponentSolver(UPrimitiveComponent* InComponent) const
+	{
+		return nullptr;
+	}
+
+
+	
+	Chaos::FPhysicsSolverEvents* FFleshCacheAdapter::BuildEventsSolver(UPrimitiveComponent* InComponent) const
+	{
+		if(UFleshComponent* FleshComp = CastChecked<UFleshComponent>(InComponent))
+		{
+			// Initialize the physics solver at the beginning of the play/record
+			FleshComp->RecreatePhysicsState();
+			return GetDeformableSolver(InComponent);
+		}
+		return nullptr;
+	}
+	
+	
+	void FFleshCacheAdapter::SetRestState(UPrimitiveComponent* InComponent, UChaosCache* InCache, const FTransform& InRootTransform, Chaos::FReal InTime) const
+	{
+		if (!InCache || InCache->GetDuration() == 0.0f)
+		{
+			return;
+		} 
+
+		if (UFleshComponent* FleshComp = CastChecked<UFleshComponent>(InComponent))
+		{
+			FPlaybackTickRecord TickRecord;
+			TickRecord.SetLastTime(InTime);
+
+			FCacheEvaluationContext Context(TickRecord);
+			Context.bEvaluateTransform = false;
+			Context.bEvaluateCurves = true;
+			Context.bEvaluateEvents = false;
+
+			FCacheEvaluationResult EvaluatedResult = InCache->Evaluate(Context, nullptr);
+
+			FleshComp->ResetDynamicCollection();
+			if (UFleshDynamicAsset* DynamicCollection = FleshComp->GetDynamicCollection())
+			{
+				TManagedArray<FVector3f>& DynamicVertex = DynamicCollection->GetPositions();
+				//TManagedArray<FVector3f>& DynamicVertex = DynamicCollection->GetVelocities();
+
+				const int32 NumCurves = EvaluatedResult.Curves.Num();
+				if (NumCurves == EvaluatedResult.ParticleIndices.Num())
+				{
+					if (DynamicVertex.Num() == NumCurves)
+					{
+						for (int32 ParticleIndex = 0; ParticleIndex < NumCurves; ++ParticleIndex)
+						{
+							if (ParticleIndex < EvaluatedResult.ParticleIndices.Num())
+							{
+								const int32 TransformIndex = EvaluatedResult.ParticleIndices[ParticleIndex];
+								if (EvaluatedResult.Curves.IsValidIndex(TransformIndex))
+								{
+									/*
+									if (const float* VelocityX = EvaluatedResult.Curves[TransformIndex].Find(VelocityXName))
+									{
+										SimulationContext->CachedVelocities[ParticleIndex].X = *VelocityX;
+									}
+									if (const float* VelocityY = EvaluatedResult.Curves[TransformIndex].Find(VelocityYName))
+									{
+										SimulationContext->CachedVelocities[ParticleIndex].Y = *VelocityY;
+									}
+									if (const float* VelocityZ = EvaluatedResult.Curves[TransformIndex].Find(VelocityZName))
+									{
+										SimulationContext->CachedVelocities[ParticleIndex].Z = *VelocityZ;
+									}
+									*/
+									if (const float* PositionX = EvaluatedResult.Curves[TransformIndex].Find(PositionXName))
+									{
+										DynamicVertex[ParticleIndex].X = *PositionX;
+									}
+									if (const float* PositionY = EvaluatedResult.Curves[TransformIndex].Find(PositionYName))
+									{
+										DynamicVertex[ParticleIndex].Y = *PositionY;
+									}
+									if (const float* PositionZ = EvaluatedResult.Curves[TransformIndex].Find(PositionZName))
+									{
+										DynamicVertex[ParticleIndex].Z = *PositionZ;
+									}
+									auto UEVertd = [](Chaos::FVec3 V) { return FVector3d(V.X, V.Y, V.Z); };
+									auto UEVertf = [](FVector3d V) { return FVector3f((float)V.X, (float)V.Y, (float)V.Z); };
+									DynamicVertex[ParticleIndex] = UEVertf(FleshComp->GetComponentTransform().InverseTransformPosition(UEVertd(DynamicVertex[ParticleIndex])));
+								}
+							}
+						}
+					}
+				}
+			}
+		}
+	}
+
+	bool FFleshCacheAdapter::InitializeForRecord(UPrimitiveComponent* InComponent, UChaosCache* InCache)
+	{
+		if( FDeformableSolver* Solver = GetDeformableSolver(InComponent))
+		{
+			FDeformableSolver::FGameThreadAccess GameThreadAccess(Solver, Softs::FGameThreadAccessor());
+			GameThreadAccess.SetEnableSolver(true);
+			if (UFleshComponent* FleshComp = CastChecked<UFleshComponent>(InComponent))
+			{
+				FleshComp->ResetDynamicCollection();
+			}
+		}
+		return true;
+	}
+
+	bool FFleshCacheAdapter::InitializeForPlayback(UPrimitiveComponent* InComponent, UChaosCache* InCache, float InTime)
+	{
+		EnsureIsInGameThreadContext();
+		
+		if (FDeformableSolver* Solver = GetDeformableSolver(InComponent))
+		{
+			FDeformableSolver::FGameThreadAccess GameThreadAccess(Solver, Softs::FGameThreadAccessor());
+			GameThreadAccess.SetEnableSolver(false);
+			if (UFleshComponent* FleshComp = CastChecked<UFleshComponent>(InComponent))
+			{
+				FleshComp->ResetDynamicCollection();
+			}
+		}
+		return true;
+	}
+}    // namespace Chaos
diff --git a/Engine/Plugins/Experimental/ChaosFlesh/Source/ChaosFleshEngine/Private/ChaosFlesh/ChaosDeformablePhysicsComponent.cpp b/Engine/Plugins/Experimental/ChaosFlesh/Source/ChaosFleshEngine/Private/ChaosFlesh/ChaosDeformablePhysicsComponent.cpp
new file mode 100644
--- /dev/null
+++ b/Engine/Plugins/Experimental/ChaosFlesh/Source/ChaosFleshEngine/Private/ChaosFlesh/ChaosDeformablePhysicsComponent.cpp
@@ -0,0 +1,180 @@
+// Copyright Epic Games, Inc. All Rights Reserved.
+
+#include "ChaosFlesh/ChaosDeformablePhysicsComponent.h"
+
+#include "Animation/SkeletalMeshActor.h"
+#include "Chaos/Deformable/ChaosDeformableSolverProxy.h"
+#include "ChaosFlesh/ChaosDeformableSolverActor.h"
+#include "ChaosFlesh/ChaosDeformableSolverComponent.h"
+
+
+DEFINE_LOG_CATEGORY_STATIC(LogDeformablePhysicsComponentInternal, Log, All);
+
+void UDeformablePhysicsComponent::OnCreatePhysicsState()
+{
+	Super::OnCreatePhysicsState();
+	UE_LOG(LogDeformablePhysicsComponentInternal, Log, TEXT("UDeformablePhysicsComponent::OnCreatePhysicsState()"));
+	if (PrimarySolver)
+	{
+		if (UDeformableSolverComponent* SolverComponent = PrimarySolver->GetDeformableSolverComponent())
+		{
+			FDeformableSolver::FGameThreadAccess GameThreadSolver = SolverComponent->GameThreadAccess();
+			if (GameThreadSolver())
+			{
+				AddProxy(GameThreadSolver);
+			}
+		}
+	}
+}
+
+void UDeformablePhysicsComponent::OnDestroyPhysicsState()
+{
+	Super::OnDestroyPhysicsState();
+	UE_LOG(LogDeformablePhysicsComponentInternal, Log, TEXT("UDeformablePhysicsComponent::OnDestroyPhysicsState()"));
+	if (PrimarySolver)
+	{
+		if (UDeformableSolverComponent* SolverComponent = PrimarySolver->GetDeformableSolverComponent())
+		{
+			FDeformableSolver::FGameThreadAccess GameThreadSolver = SolverComponent->GameThreadAccess();
+			if (GameThreadSolver())
+			{
+				RemoveProxy(GameThreadSolver);
+			}
+		}
+	}
+}
+
+bool UDeformablePhysicsComponent::ShouldCreatePhysicsState() const
+{
+	return true;
+}
+bool UDeformablePhysicsComponent::HasValidPhysicsState() const
+{
+	return PhysicsProxy != nullptr;
+}
+
+
+#if WITH_EDITOR
+void UDeformablePhysicsComponent::PreEditChange(FProperty* PropertyThatWillChange)
+{
+	Super::PreEditChange(PropertyThatWillChange);
+	if (PropertyThatWillChange && PropertyThatWillChange->GetFName() == GET_MEMBER_NAME_CHECKED(UFleshComponent, PrimarySolver))
+	{
+		PreEditChangePrimarySolver = PrimarySolver;
+	}
+}
+
+
+void UDeformablePhysicsComponent::PostEditChangeProperty(struct FPropertyChangedEvent& PropertyChangedEvent)
+{
+	Super::PostEditChangeProperty(PropertyChangedEvent);
+
+	//
+	// The UDeformablePhysicsComponent and the UDeformableSolverComponent hold references to each other. 
+	// If one of the attributes change, then the attribute on the other component needs to be updated. 
+	//
+	if (PropertyChangedEvent.Property && PropertyChangedEvent.Property->GetFName() == GET_MEMBER_NAME_CHECKED(UFleshComponent, PrimarySolver))
+	{
+		if (AActor* Owner = Cast<AActor>(GetOwner()))
+		{
+			if (PrimarySolver)
+			{
+				if (UDeformableSolverComponent* SolverComponent = PrimarySolver->GetDeformableSolverComponent())
+				{
+					if (!SolverComponent->DeformableActors.Contains(Owner))
+					{
+						SolverComponent->DeformableActors.Add(TObjectPtr<AActor>(Owner));
+					}
+				}
+			}
+			else if (PreEditChangePrimarySolver)
+			{
+				bool DrivesOtherComponentsOnActor = false;
+				TArray<UDeformablePhysicsComponent*> DeformableComponentsOnActor;
+				Owner->GetComponents<UDeformablePhysicsComponent>(DeformableComponentsOnActor);
+				for (UDeformablePhysicsComponent* DeformableComponent : DeformableComponentsOnActor)
+				{
+					if (DeformableComponent->PrimarySolver == PreEditChangePrimarySolver)
+					{
+						DrivesOtherComponentsOnActor = true;
+						break;
+					}
+				}
+				if (!DrivesOtherComponentsOnActor)
+				{
+					if (UDeformableSolverComponent* SolverComponent = PreEditChangePrimarySolver->GetDeformableSolverComponent())
+					{
+						if (SolverComponent->DeformableActors.Contains(Owner))
+						{
+							SolverComponent->DeformableActors.Remove(Owner);
+						}
+					}
+				}
+			}
+		}
+		PreEditChangePrimarySolver = nullptr;
+	}
+}
+#endif
+
+
+void UDeformablePhysicsComponent::AddProxy(Chaos::Softs::FDeformableSolver::FGameThreadAccess& GameThreadSolver)
+{
+	PhysicsProxy = NewProxy();
+	if (PhysicsProxy)
+	{
+		// PhysicsProxy is created on game thread but is owned by physics thread. This is the handoff. 
+		GameThreadSolver.AddProxy(PhysicsProxy);
+	}
+}
+
+void UDeformablePhysicsComponent::RemoveProxy(Chaos::Softs::FDeformableSolver::FGameThreadAccess& GameThreadSolver)
+{
+	if (PhysicsProxy)
+	{
+		GameThreadSolver.RemoveProxy(PhysicsProxy);
+		PhysicsProxy = nullptr; // destroyed on physics thread. 
+	}
+}
+
+UDeformableSolverComponent* UDeformablePhysicsComponent::GetDeformableSolver()
+{
+	return PrimarySolver ? PrimarySolver->GetDeformableSolverComponent() : nullptr;
+}
+const UDeformableSolverComponent* UDeformablePhysicsComponent::GetDeformableSolver() const
+{ 
+	return PrimarySolver ? PrimarySolver->GetDeformableSolverComponent() : nullptr; 
+}
+
+UDeformablePhysicsComponent::UDeformablePhysicsComponent(const FObjectInitializer& ObjectInitializer)
+	: Super(ObjectInitializer)
+{
+}
+
+
+void UDeformablePhysicsComponent::EnableSimulation(ADeformableSolverActor* InActor)
+{
+	if (InActor)
+	{
+		if (UDeformableSolverComponent* SolverComponent = InActor->GetDeformableSolverComponent())
+		{
+			if (AActor* Owner = Cast<AActor>(this->GetOwner()))
+			{
+				PrimarySolver = InActor;
+				if (!SolverComponent->DeformableActors.Contains(Owner))
+				{
+					SolverComponent->DeformableActors.Add(Owner);
+				}
+				SolverComponent->AddDeformableProxy(this);
+			}
+		}
+	}
+}
+
+
+
+
+
+
+
+
diff --git a/Engine/Plugins/Experimental/ChaosFlesh/Source/ChaosFleshEngine/Private/ChaosFlesh/ChaosDeformableSolverActor.cpp b/Engine/Plugins/Experimental/ChaosFlesh/Source/ChaosFleshEngine/Private/ChaosFlesh/ChaosDeformableSolverActor.cpp
new file mode 100644
--- /dev/null
+++ b/Engine/Plugins/Experimental/ChaosFlesh/Source/ChaosFleshEngine/Private/ChaosFlesh/ChaosDeformableSolverActor.cpp
@@ -0,0 +1,20 @@
+// Copyright Epic Games, Inc. All Rights Reserved.
+
+#include "ChaosFlesh/ChaosDeformableSolverActor.h"
+
+#include "ChaosFlesh/FleshActor.h"
+#include "ChaosFlesh/ChaosDeformableSolverComponent.h"
+#include "ChaosFlesh/ChaosDeformableSolverThreading.h"
+#include "CoreMinimal.h"
+
+DEFINE_LOG_CATEGORY_STATIC(LogChaosDeformableSolverActor, Log, All);
+
+ADeformableSolverActor::ADeformableSolverActor(const FObjectInitializer& ObjectInitializer)
+	: Super(ObjectInitializer)
+{
+	//UE_LOG(ADeformableSolverActor, Verbose, TEXT("ADeformableSolverActor::ADeformableSolverActor()"));
+	SolverComponent = CreateDefaultSubobject<UDeformableSolverComponent>(TEXT("DeformableSolverComponent0"));
+	RootComponent = SolverComponent;
+}
+
+
diff --git a/Engine/Plugins/Experimental/ChaosFlesh/Source/ChaosFleshEngine/Private/ChaosFlesh/ChaosDeformableSolverAsset.cpp b/Engine/Plugins/Experimental/ChaosFlesh/Source/ChaosFleshEngine/Private/ChaosFlesh/ChaosDeformableSolverAsset.cpp
new file mode 100644
--- /dev/null
+++ b/Engine/Plugins/Experimental/ChaosFlesh/Source/ChaosFleshEngine/Private/ChaosFlesh/ChaosDeformableSolverAsset.cpp
@@ -0,0 +1,12 @@
+// Copyright Epic Games, Inc. All Rights Reserved.
+
+#include "ChaosFlesh/ChaosDeformableSolverAsset.h"
+
+
+UChaosDeformableSolver::UChaosDeformableSolver(const FObjectInitializer& ObjectInitializer)
+	: UObject(ObjectInitializer)
+{
+	check(ObjectInitializer.GetClass() == GetClass());
+}
+
+
diff --git a/Engine/Plugins/Experimental/ChaosFlesh/Source/ChaosFleshEngine/Private/ChaosFlesh/ChaosDeformableSolverComponent.cpp b/Engine/Plugins/Experimental/ChaosFlesh/Source/ChaosFleshEngine/Private/ChaosFlesh/ChaosDeformableSolverComponent.cpp
new file mode 100644
--- /dev/null
+++ b/Engine/Plugins/Experimental/ChaosFlesh/Source/ChaosFleshEngine/Private/ChaosFlesh/ChaosDeformableSolverComponent.cpp
@@ -0,0 +1,359 @@
+// Copyright Epic Games, Inc. All Rights Reserved.
+
+#include "ChaosFlesh/ChaosDeformableSolverComponent.h"
+
+#include "ChaosFlesh/ChaosDeformableSolverActor.h"
+#include "ChaosFlesh/ChaosDeformableSolverThreading.h"
+#include "ChaosFlesh/FleshActor.h"
+#include "ChaosFlesh/FleshComponent.h"
+#include "Animation/SkeletalMeshActor.h"
+#include "ChaosFlesh/FleshComponent.h"
+#include "CoreMinimal.h"
+#include "Engine/World.h"
+
+DEFINE_LOG_CATEGORY_STATIC(LogDeformableSolverComponentInternal, Log, All);
+
+static TAutoConsoleVariable<int32> CVarDeformablePhysicsTickWaitForParallelDeformableTask(
+	TEXT("p.ClothPhysics.WaitForParallelDeformableTask"), 0, 
+	TEXT("If 1, always wait for deformable task completion in the Deformable Tick function. "\
+		 "If 0, wait at end - of - frame updates instead if allowed by component settings"));
+
+UDeformableSolverComponent::UDeformableSolverComponent(const FObjectInitializer& ObjectInitializer)
+	: Super(ObjectInitializer)
+	, Solver()
+{
+	PrimaryComponentTick.bCanEverTick = true;
+	bTickInEditor = false;
+
+	PrimaryComponentTick.bCanEverTick = true;
+	PrimaryComponentTick.bTickEvenWhenPaused = false;
+	PrimaryComponentTick.TickGroup = TG_PrePhysics;
+
+	DeformableEndTickFunction.TickGroup = TG_PostPhysics;
+	DeformableEndTickFunction.bCanEverTick = true;
+	DeformableEndTickFunction.bStartWithTickEnabled = true;
+}
+
+UDeformableSolverComponent::~UDeformableSolverComponent()
+{
+}
+
+UDeformableSolverComponent::FDeformableSolver::FGameThreadAccess
+UDeformableSolverComponent::GameThreadAccess()
+{
+	return FDeformableSolver::FGameThreadAccess(Solver.Get(), Chaos::Softs::FGameThreadAccessor());
+}
+
+
+#if WITH_EDITOR
+void UDeformableSolverComponent::PreEditChange(FProperty* PropertyThatWillChange)
+{
+	Super::PreEditChange(PropertyThatWillChange);
+	if (PropertyThatWillChange && PropertyThatWillChange->GetFName() == GET_MEMBER_NAME_CHECKED(UDeformableSolverComponent, DeformableActors))
+	{
+		PreEditChangeDeformableActors = DeformableActors;
+	}
+}
+
+void UDeformableSolverComponent::PostEditChangeProperty(struct FPropertyChangedEvent& PropertyChangedEvent)
+{
+	Super::PostEditChangeProperty(PropertyChangedEvent);
+
+	//
+	// The UDeformablePhysicsComponent and the UDeformableSolverComponent hold references to each other. 
+	// If one of the attributes change, then the attribute on the other component needs to be updated. 
+	//
+	if (PropertyChangedEvent.Property && PropertyChangedEvent.Property->GetFName() == GET_MEMBER_NAME_CHECKED(UDeformableSolverComponent, DeformableActors))
+	{
+		// process removed actors
+		for (TObjectPtr<AActor>& SimulatedActor : PreEditChangeDeformableActors)
+		{
+			if (SimulatedActor)
+			{
+				if (!DeformableActors.Contains(SimulatedActor))
+				{
+					TArray<UDeformablePhysicsComponent*> DeformableComponentsOnActor;
+					SimulatedActor->GetComponents<UDeformablePhysicsComponent>(DeformableComponentsOnActor);
+					for (UDeformablePhysicsComponent* DeformableComponent : DeformableComponentsOnActor)
+					{
+						if (DeformableComponent->PrimarySolver == GetOwner())
+						{
+							DeformableComponent->PrimarySolver = nullptr;
+						}
+					}
+				}
+			}
+		}
+		// process added actors
+		for (TObjectPtr<AActor>& SimulatedActor : DeformableActors)
+		{
+			if (SimulatedActor)
+			{
+				if (!PreEditChangeDeformableActors.Contains(SimulatedActor))
+				{
+					TArray<UDeformablePhysicsComponent*> DeformableComponentsOnActor;
+					SimulatedActor->GetComponents<UDeformablePhysicsComponent>(DeformableComponentsOnActor);
+					for (UDeformablePhysicsComponent* DeformableComponent : DeformableComponentsOnActor)
+					{
+						if (DeformableComponent->PrimarySolver == nullptr)
+						{
+							DeformableComponent->PrimarySolver = Cast<ADeformableSolverActor>(GetOwner());
+						}
+					}
+				}
+			}
+		}
+	}
+}
+#endif
+
+bool UDeformableSolverComponent::IsSimulatable() const
+{
+	return true;
+}
+
+bool UDeformableSolverComponent::IsSimulating(UDeformablePhysicsComponent* InComponent) const
+{
+	if (InComponent)
+	{
+		const AActor* ComponentSolver = InComponent->PrimarySolver.Get();
+		return ComponentSolver == GetOwner();
+	}
+	return false;
+}
+
+
+void UDeformableSolverComponent::UpdateDeformableEndTickState(bool bRegister)
+{
+	UE_LOG(LogDeformableSolverComponentInternal, Verbose, TEXT("UDeformableSolverComponent::RegiUpdateDeformableEndTickStatesterEndPhysicsTick"));
+	bRegister &= PrimaryComponentTick.IsTickFunctionRegistered();
+
+	if (bRegister != DeformableEndTickFunction.IsTickFunctionRegistered())
+	{
+		if (bRegister)
+		{
+			UWorld* World = GetWorld();
+			if (World->EndPhysicsTickFunction.IsTickFunctionRegistered() && SetupActorComponentTickFunction(&DeformableEndTickFunction))
+			{
+				DeformableEndTickFunction.DeformableSolverComponent = this;
+				// Make sure our EndPhysicsTick gets called after physics simulation is finished
+				if (World != nullptr)
+				{
+					DeformableEndTickFunction.AddPrerequisite(World, World->EndPhysicsTickFunction);
+				}
+			}
+		}
+		else
+		{
+			DeformableEndTickFunction.UnRegisterTickFunction();
+		}
+	}
+}
+
+//bool UDeformableSolverComponent::ShouldWaitForDeformableInTickFunction() const
+//{
+//	return bWaitForParallelDeformableTask || (CVarDeformablePhysicsTickWaitForParallelDeformableTask.GetValueOnAnyThread() != 0);
+//}
+
+void UDeformableSolverComponent::BeginPlay()
+{
+	Super::BeginPlay();
+	Reset();
+}
+
+void UDeformableSolverComponent::TickComponent(float DeltaTime, enum ELevelTick TickType, FActorComponentTickFunction* ThisTickFunction)
+{
+	UE_LOG(LogDeformableSolverComponentInternal, Verbose, TEXT("UDeformableSolverComponent::TickComponent"));
+	TRACE_CPUPROFILER_EVENT_SCOPE(DeformableSolverComponent_TickComponent);
+
+	UpdateDeformableEndTickState(IsSimulatable());
+
+	UpdateFromGameThread(DeltaTime);
+
+	if (bDoThreadedAdvance)
+	{
+		//ParallelDeformableTask = TGraphTask<FParallelDeformableTask>::CreateTask(nullptr, ENamedThreads::GameThread).ConstructAndDispatchWhenReady(this, DeltaTime);
+		//if (ShouldWaitForDeformableInTickFunction())
+		{ // see FParallelClothCompletionTask
+			FGraphEventArray Prerequisites;
+			Prerequisites.Add(ParallelDeformableTask);
+			FGraphEventRef DeformableCompletionEvent = TGraphTask<FParallelDeformableTask>::CreateTask(&Prerequisites, ENamedThreads::GameThread).ConstructAndDispatchWhenReady(this, DeltaTime);
+			ThisTickFunction->GetCompletionHandle()->DontCompleteUntil(DeformableCompletionEvent);
+		}
+	}
+	else
+	{
+		Simulate(DeltaTime);
+	}
+
+
+}
+
+void UDeformableSolverComponent::EndPlay(const EEndPlayReason::Type EndPlayReason)
+{
+	Super::EndPlay(EndPlayReason);
+}
+
+void UDeformableSolverComponent::Reset()
+{
+	Solver.Reset(new FDeformableSolver({
+		NumSubSteps
+		, NumSolverIterations
+		, FixTimeStep
+		, TimeStepSize
+		, CacheToFile
+		, bEnableKinematics
+		, bUseFloor
+		, bDoSelfCollision
+		, bUseGridBasedConstraints
+		, GridDx
+		}));
+
+	
+	for (TObjectPtr<AActor>& DeformableActor : DeformableActors)
+	{
+		if (DeformableActor)
+		{
+			TArray<UDeformablePhysicsComponent*> DeformableComponents;
+			DeformableActor->GetComponents<UDeformablePhysicsComponent>(DeformableComponents);
+
+			for (UDeformablePhysicsComponent* DeformableComponent : DeformableComponents)
+			{
+				if (IsSimulating(DeformableComponent))
+				{
+					AddDeformableProxy(DeformableComponent);
+				}
+			}
+		}
+	}
+	
+}
+
+/*
+* @todo(flesh) : In game additions. 
+void UDeformableSolverComponent::AddDeformableActor(TObjectPtr<AActor> InActor)
+{
+	if (InActor)
+	{
+		if (!DeformableActors.Contains(InActor))
+		{
+			DeformableActors.Add(InActor);
+		}
+
+		if (AFleshActor* FleshActor = Cast<AFleshActor>(InActor.Get()))
+		{
+			TArray<UDeformablePhysicsComponent*> DeformableComponents;
+			InActor->GetComponents<UDeformablePhysicsComponent>(DeformableComponents);
+
+			for (UDeformablePhysicsComponent* DeformableComponent : DeformableComponents)
+			{
+				AddDeformableProxy(DeformableComponent);
+			}
+		}
+	}
+}
+*/
+
+void UDeformableSolverComponent::AddDeformableProxy(UDeformablePhysicsComponent* InComponent)
+{
+	TRACE_CPUPROFILER_EVENT_SCOPE(DeformableSolverComponent_AddDeformableProxy);
+
+	if (Solver && IsSimulating(InComponent))
+	{
+		FDeformableSolver::FGameThreadAccess GameThreadSolver(Solver.Get(), Chaos::Softs::FGameThreadAccessor());
+		if (!GameThreadSolver.HasObject(InComponent))
+		{
+			InComponent->AddProxy(GameThreadSolver);
+		}
+	}
+}
+
+void UDeformableSolverComponent::Simulate(float DeltaTime)
+{
+	TRACE_CPUPROFILER_EVENT_SCOPE(DeformableSolverComponent_Simulate);
+
+	if (Solver)
+	{
+		// @todo(accessor) : Should be coming from the threading class. 
+		FDeformableSolver::FPhysicsThreadAccess PhysicsThreadSolver(Solver.Get(), Chaos::Softs::FPhysicsThreadAccessor());
+		PhysicsThreadSolver.Simulate(DeltaTime);
+	}
+}
+
+void UDeformableSolverComponent::UpdateFromGameThread(float DeltaTime)
+{
+	TRACE_CPUPROFILER_EVENT_SCOPE(DeformableSolverComponent_UpdateFromGameThread);
+
+	if (Solver)
+	{
+		FDeformableSolver::FGameThreadAccess GameThreadSolver(Solver.Get(), Chaos::Softs::FGameThreadAccessor());
+
+		Chaos::Softs::FDeformableDataMap DataMap;
+		for (TObjectPtr<AActor>& DeformableActor : DeformableActors)
+		{
+			if (DeformableActor)
+			{
+				TArray<UDeformablePhysicsComponent*> DeformableComponents;
+				DeformableActor->GetComponents<UDeformablePhysicsComponent>(DeformableComponents);
+
+				for (UDeformablePhysicsComponent* DeformableComponent : DeformableComponents)
+				{
+					if (IsSimulating(DeformableComponent))
+					{
+						if (FDataMapValue Value = DeformableComponent->NewDeformableData())
+						{
+							DataMap.Add(DeformableComponent, Value);
+						}
+					}
+				}
+			}
+		}
+
+		GameThreadSolver.PushInputPackage(GameThreadSolver.GetFrame(), MoveTemp(DataMap));
+	}
+}
+
+void UDeformableSolverComponent::UpdateFromSimulation(float DeltaTime)
+{
+	TRACE_CPUPROFILER_EVENT_SCOPE(DeformableSolverComponent_UpdateFromSimulation);
+
+	if (Solver)
+	{
+		FDeformableSolver::FGameThreadAccess GameThreadSolver(Solver.Get(), Chaos::Softs::FGameThreadAccessor());
+		TUniquePtr<FDeformablePackage> Output(nullptr);
+		while (TUniquePtr<FDeformablePackage> SolverOutput = GameThreadSolver.PullOutputPackage())
+		{
+			Output.Reset(SolverOutput.Release());
+		}
+
+		if (Output)
+		{
+			for (TObjectPtr<AActor>& DeformableActor : DeformableActors)
+			{
+				if (DeformableActor)
+				{
+					TArray<UDeformablePhysicsComponent*> DeformableComponents;
+					DeformableActor->GetComponents<UDeformablePhysicsComponent>(DeformableComponents);
+
+					for (UDeformablePhysicsComponent* DeformableComponent : DeformableComponents)
+					{
+						if (IsSimulating(DeformableComponent))
+						{
+							if (const FDataMapValue* Buffer = Output->ObjectMap.Find(DeformableComponent))
+							{
+								DeformableComponent->UpdateFromSimualtion(Buffer);
+							}
+						}
+					}
+				}
+			}
+		}
+	}
+}
+
+
+
+
+
+
+
diff --git a/Engine/Plugins/Experimental/ChaosFlesh/Source/ChaosFleshEngine/Private/ChaosFlesh/FleshActor.cpp b/Engine/Plugins/Experimental/ChaosFlesh/Source/ChaosFleshEngine/Private/ChaosFlesh/FleshActor.cpp
new file mode 100644
--- /dev/null
+++ b/Engine/Plugins/Experimental/ChaosFlesh/Source/ChaosFleshEngine/Private/ChaosFlesh/FleshActor.cpp
@@ -0,0 +1,72 @@
+// Copyright Epic Games, Inc. All Rights Reserved.
+
+#include "ChaosFlesh/FleshActor.h"
+
+#include "CoreMinimal.h"
+#include "Containers/StringConv.h"
+#include "Chaos/TriangleMesh.h"
+
+#include "Chaos/PBDBendingConstraints.h"
+#include "Chaos/PBDParticles.h"
+#include "Chaos/PBDSpringConstraints.h"
+#include "Chaos/PBDVolumeConstraint.h"
+#include "Chaos/PerParticleGravity.h"
+#include "Chaos/PBDSpringConstraints.h"
+#include "Chaos/PBDTriangleMeshCollisions.h"
+#include "Chaos/PBDCollisionSpringConstraints.h"
+#include "Chaos/XPBDCorotatedConstraints.h"
+#include "Chaos/XPBDVolumeConstraints.h"
+#include "Chaos/XPBDCorotatedFiberConstraints.h"
+#include "Chaos/PBDTetConstraints.h"
+#include "Chaos/PBDAltitudeSpringConstraints.h"
+#include "Chaos/Plane.h"
+#include "Chaos/Utilities.h"
+#include "Chaos/PBDEvolution.h"
+#include "Misc/FileHelper.h"
+#include "HAL/PlatformFileManager.h"
+#include "Misc/Paths.h"
+//#include "ChaosFlesh/PB.h"
+
+#include <fstream>
+#include <functional>
+#include <sstream>
+#include <unordered_map>
+#include <unordered_set>
+
+#include UE_INLINE_GENERATED_CPP_BY_NAME(FleshActor)
+
+DEFINE_LOG_CATEGORY_STATIC(AFleshLogging, Log, All);
+
+AFleshActor::AFleshActor(const FObjectInitializer& ObjectInitializer)
+	: Super(ObjectInitializer)
+{
+	//UE_LOG(AFleshLogging, Verbose, TEXT("AFleshActor::AFleshActor()"));
+	FleshComponent = CreateDefaultSubobject<UFleshComponent>(TEXT("FleshComponent0"));
+	RootComponent = FleshComponent;
+	PrimaryActorTick.bCanEverTick = true;
+}
+
+#if WITH_EDITOR
+bool AFleshActor::GetReferencedContentObjects(TArray<UObject*>& Objects) const
+{
+	Super::GetReferencedContentObjects(Objects);
+
+	if (FleshComponent && FleshComponent->GetRestCollection())
+	{
+		Objects.Add(const_cast<UFleshAsset*>(FleshComponent->GetRestCollection()));
+	}
+	return true;
+}
+#endif
+
+void AFleshActor::EnableSimulation(ADeformableSolverActor* InActor)
+{
+	if (InActor)
+	{
+		if (FleshComponent && FleshComponent->GetRestCollection())
+		{
+			FleshComponent->EnableSimulation(InActor);
+		}
+	}
+}
+
diff --git a/Engine/Plugins/Experimental/ChaosFlesh/Source/ChaosFleshEngine/Private/ChaosFlesh/FleshComponent.cpp b/Engine/Plugins/Experimental/ChaosFlesh/Source/ChaosFleshEngine/Private/ChaosFlesh/FleshComponent.cpp
new file mode 100644
--- /dev/null
+++ b/Engine/Plugins/Experimental/ChaosFlesh/Source/ChaosFleshEngine/Private/ChaosFlesh/FleshComponent.cpp
@@ -0,0 +1,361 @@
+// Copyright Epic Games, Inc. All Rights Reserved.
+
+#include "ChaosFlesh/FleshComponent.h"
+
+#include "Animation/SkeletalMeshActor.h"
+#include "Chaos/Deformable/ChaosDeformableSolverProxy.h"
+#include "ChaosFlesh/FleshCollectionUtility.h"
+#include "ChaosFlesh/ChaosDeformableSolverActor.h"
+#include "ChaosFlesh/ChaosDeformableSolverComponent.h"
+#include "Components/SkeletalMeshComponent.h"
+#include "Dataflow/DataflowEngineUtil.h"
+#include "Engine/SkeletalMesh.h"
+#include "GeometryCollection/Facades/CollectionTransformSourceFacade.h"
+
+#include UE_INLINE_GENERATED_CPP_BY_NAME(FleshComponent)
+
+DEFINE_LOG_CATEGORY_STATIC(LogFleshComponentInternal, Log, All);
+
+UFleshComponent::UFleshComponent(const FObjectInitializer& ObjectInitializer)
+	: Super(ObjectInitializer)
+{
+	PrimaryComponentTick.bCanEverTick = true;
+	bTickInEditor = true;
+	DynamicCollection = ObjectInitializer.CreateDefaultSubobject<UFleshDynamicAsset>(this, TEXT("Flesh Dynamic Asset"));
+	Mesh = ObjectInitializer.CreateDefaultSubobject<UProceduralMeshComponent>(this, TEXT("Flesh Visualization Component"));
+}
+
+UFleshComponent::~UFleshComponent()
+{
+	if (RenderMesh) delete RenderMesh;
+}
+
+void UFleshComponent::Invalidate()
+{
+	bBoundsNeedsUpdate = true;
+}
+
+
+void UFleshComponent::OnRegister()
+{
+	if (bBoundsNeedsUpdate)
+	{
+		UpdateLocalBounds();
+	}
+	Super::OnRegister();
+}
+
+
+void UFleshComponent::BeginPlay()
+{
+	Super::BeginPlay();
+
+	if (ASkeletalMeshActor* SkeletalMeshActor = Cast<ASkeletalMeshActor>(GetOwner()))
+	{
+		if (USkeletalMeshComponent* SkeletalMeshComponent = SkeletalMeshActor->GetSkeletalMeshComponent())
+		{
+			PrimaryComponentTick.AddPrerequisite(SkeletalMeshComponent, SkeletalMeshComponent->PrimaryComponentTick);
+		}
+	}
+}
+
+void UFleshComponent::TickComponent(float DeltaTime, enum ELevelTick TickType, FActorComponentTickFunction *ThisTickFunction)
+{
+	Super::TickComponent(DeltaTime, TickType, ThisTickFunction);
+	//UE_LOG(LogFleshComponentInternal, Log, TEXT("UFleshComponent::TickComponent()"));
+	
+	RenderProceduralMesh();
+	UpdateLocalBounds();
+}
+
+void UFleshComponent::EndPlay(const EEndPlayReason::Type ReasonEnd)
+{
+	if (GetDynamicCollection())
+	{
+		GetDynamicCollection()->Reset();
+	}
+}
+
+void UFleshComponent::SetRestCollection(const UFleshAsset* InRestCollection)
+{
+	RestCollection = InRestCollection;
+	Invalidate();
+	UpdateLocalBounds();
+	ResetProceduralMesh();
+}
+
+
+UDeformablePhysicsComponent::FThreadingProxy* UFleshComponent::NewProxy()
+{
+	TRACE_CPUPROFILER_EVENT_SCOPE(DeformableSolverComponent_AddFleshProxy);
+	if (const UFleshAsset* RestAsset = GetRestCollection())
+	{
+		if (const FFleshCollection* Rest = RestAsset->GetCollection())
+		{
+			if (!GetDynamicCollection())
+			{
+				DynamicCollection = NewObject<UFleshDynamicAsset>(this, TEXT("Flesh Dynamic Asset"));
+			}
+
+			GetDynamicCollection()->Reset(Rest);
+			if (const FManagedArrayCollection* Dynamic = GetDynamicCollection()->GetCollection())
+			{
+				return new FFleshThreadingProxy(this, GetComponentTransform(), *Rest, *Dynamic);
+			}
+		}
+	}
+	return nullptr;
+}
+
+UDeformablePhysicsComponent::FDataMapValue UFleshComponent::NewDeformableData()
+{
+	using namespace GeometryCollection::Facades;
+	if (GetOwner())
+	{
+		if (const UFleshAsset* FleshAsset = GetRestCollection())
+		{
+			if (const FFleshCollection* Rest = FleshAsset->GetCollection())
+			{
+				if (FTransformSource::HasFacade(Rest))
+				{
+
+					TArray<USkeletalMeshComponent*> SkeletalMeshComponents;
+					GetOwner()->GetComponents<USkeletalMeshComponent>(SkeletalMeshComponents);
+
+					if (const TManagedArray<FTransform>* RestTransforms = Rest->FindAttribute<FTransform>(FTransformCollection::TransformAttribute, FTransformCollection::TransformGroup))
+					{
+						TArray<FTransform> AnimationTransforms = RestTransforms->GetConstArray();
+						TArray<FTransform> ComponentPose = RestTransforms->GetConstArray();
+
+						for (const USkeletalMeshComponent* SkeletalMeshComponent : SkeletalMeshComponents)
+						{
+							if (const USkeletalMesh* SkeletalMesh = SkeletalMeshComponent->GetSkeletalMeshAsset())
+							{
+								if (const USkeleton* Skeleton = SkeletalMesh->GetSkeleton())
+								{
+									TSet<int32> Roots = FTransformSource::GetTransformSource(Rest, Skeleton->GetName(), Skeleton->GetGuid().ToString());
+									if (!Roots.IsEmpty() && ensureMsgf(Roots.Num() == 1, TEXT("Error: Only supports a single root per skeleton.(%s)"), *Skeleton->GetName()))
+									{
+										TArray<FTransform> ComponentLocalPose;
+										Dataflow::Animation::GlobalTransforms(SkeletalMesh->GetRefSkeleton(), ComponentLocalPose);
+
+										const TArray<FTransform>& ComponentTransforms = SkeletalMeshComponent->GetComponentSpaceTransforms();
+										if (ComponentLocalPose.Num() == ComponentTransforms.Num())
+										{
+											for (int Adx = Roots.Array()[0], Cdx = 0; Adx <= AnimationTransforms.Num() && Cdx < ComponentTransforms.Num(); Adx++, Cdx++)
+											{
+												// @todo(flesh) : Can we just use one array?
+												AnimationTransforms[Adx] = ComponentTransforms[Cdx];
+												ComponentPose[Adx] = ComponentLocalPose[Cdx]; 
+											}
+										}
+									}
+								}
+							}
+						}
+						return FDataMapValue(new Chaos::Softs::FFleshThreadingProxy::FFleshInputBuffer(this->GetComponentTransform(), AnimationTransforms, ComponentPose, this));
+					}
+				}
+			}
+		}
+	}
+	return FDataMapValue(new Chaos::Softs::FFleshThreadingProxy::FFleshInputBuffer(this->GetComponentTransform(), this));
+}
+
+
+void UFleshComponent::UpdateFromSimualtion(const FDataMapValue* SimualtionBuffer)
+{
+	if (const FFleshThreadingProxy::FFleshOutputBuffer* FleshBuffer = (*SimualtionBuffer)->As<FFleshThreadingProxy::FFleshOutputBuffer>())
+	{
+		if (GetDynamicCollection())
+		{
+			// @todo(flesh) : reduce conversions
+			auto UEVertd = [](FVector3f V) { return FVector3d(V.X, V.Y, V.Z); };
+			auto UEVertf = [](FVector3d V) { return FVector3f(V.X, V.Y, V.Z); };
+
+			TManagedArray<FVector3f>& DynamicVertex = GetDynamicCollection()->GetPositions();
+			const TManagedArray<FVector3f>& SimulationVertex = FleshBuffer->Dynamic.GetAttribute<FVector3f>("Vertex", FGeometryCollection::VerticesGroup);
+
+			for (int i = DynamicVertex.Num() - 1; i >= 0; i--)
+			{
+				DynamicVertex[i] = UEVertf(GetComponentTransform().InverseTransformPosition(UEVertd(SimulationVertex[i])));
+			}
+		}
+	}
+}
+
+void UFleshComponent::UpdateLocalBounds()
+{
+	if (bBoundsNeedsUpdate && RestCollection)
+	{
+		{
+			FFleshAssetEdit EditObject = RestCollection->EditCollection();
+			if (FFleshCollection* Collection = EditObject.GetFleshCollection())
+			{
+				Collection->UpdateBoundingBox();
+			}
+		}
+		BoundingBox = RestCollection->GetCollection()->GetBoundingBox();
+		bBoundsNeedsUpdate = false;
+	}
+
+}
+
+FBoxSphereBounds UFleshComponent::CalcBounds(const FTransform& LocalToWorldIn) const
+{
+	return BoundingBox.TransformBy(GetComponentTransform()); // todo(chaos:flesh) use LocalToWorldIn
+}
+
+void UFleshComponent::ResetDynamicCollection()
+{
+	if (const UFleshAsset* RestAsset = GetRestCollection())
+	{
+		if (!GetDynamicCollection())
+		{
+			DynamicCollection = NewObject<UFleshDynamicAsset>(this, TEXT("Flesh Dynamic Asset"));
+		}
+
+		if (!GetDynamicCollection()->GetCollection() || 
+			!GetDynamicCollection()->GetCollection()->NumElements(FGeometryCollection::VerticesGroup))
+		{
+			GetDynamicCollection()->Reset(RestAsset->GetCollection());
+			ResetProceduralMesh();
+		}
+		else
+		{
+			GetDynamicCollection()->ResetAttributesFrom(RestAsset->GetCollection());
+		}
+	}
+}
+
+void UFleshComponent::ResetProceduralMesh()
+{
+	if (Mesh && RenderMesh)
+	{
+		Mesh->ClearAllMeshSections();
+		delete RenderMesh;
+		RenderMesh = nullptr;
+	}
+}
+
+void UFleshComponent::RenderProceduralMesh()
+{
+	bool bCanRender = false;
+	if (const UFleshAsset* FleshAsset = GetRestCollection())
+	{
+#if WITH_EDITORONLY_DATA
+		if (FleshAsset->bRenderInEditor)
+#endif
+		{
+			if (const FFleshCollection* Flesh = FleshAsset->GetCollection())
+			{
+				int32 NumVertices = Flesh->NumElements(FGeometryCollection::VerticesGroup);
+				int32 NumFaces = Flesh->NumElements(FGeometryCollection::FacesGroup);
+				if (NumFaces && NumVertices)
+				{
+					if (RenderMesh && RenderMesh->Vertices.Num() != NumFaces * 3)
+					{
+						ResetProceduralMesh();
+					}
+
+					if (!RenderMesh)
+					{
+						RenderMesh = new FFleshRenderMesh;
+
+						for (int i = 0; i < NumFaces; ++i)
+						{
+							const auto& P1 = Flesh->Vertex[Flesh->Indices[i][0]];
+							const auto& P2 = Flesh->Vertex[Flesh->Indices[i][1]];
+							const auto& P3 = Flesh->Vertex[Flesh->Indices[i][2]];
+
+							RenderMesh->Vertices.Add(FVector(P1));
+							RenderMesh->Vertices.Add(FVector(P2));
+							RenderMesh->Vertices.Add(FVector(P3));
+
+							RenderMesh->Colors.Add(FLinearColor::White);
+							RenderMesh->Colors.Add(FLinearColor::White);
+							RenderMesh->Colors.Add(FLinearColor::White);
+
+							RenderMesh->UVs.Add(FVector2D(0, 0));
+							RenderMesh->UVs.Add(FVector2D(0, 0));
+							RenderMesh->UVs.Add(FVector2D(0, 0));
+
+							RenderMesh->Triangles.Add(3 * i);
+							RenderMesh->Triangles.Add(3 * i + 1);
+							RenderMesh->Triangles.Add(3 * i + 2);
+
+							auto Normal = -Chaos::FVec3::CrossProduct(P3 - P1, P2 - P1);
+							RenderMesh->Normals.Add(Normal);
+							RenderMesh->Normals.Add(Normal);
+							RenderMesh->Normals.Add(Normal);
+
+							auto Tangent = (P2 - P1).GetSafeNormal();
+							RenderMesh->Tangents.Add(FProcMeshTangent(Tangent[0], Tangent[1], Tangent[2]));
+							Tangent = (P3 - P2).GetSafeNormal();
+							RenderMesh->Tangents.Add(FProcMeshTangent(Tangent[0], Tangent[1], Tangent[2]));
+							Tangent = (P1 - P3).GetSafeNormal();
+							RenderMesh->Tangents.Add(FProcMeshTangent(Tangent[0], Tangent[1], Tangent[2]));
+						}
+
+						Mesh->SetRelativeTransform(GetComponentTransform());
+						Mesh->CreateMeshSection_LinearColor(0, RenderMesh->Vertices, RenderMesh->Triangles, RenderMesh->Normals, RenderMesh->UVs, RenderMesh->Colors, RenderMesh->Tangents, false);
+					}
+					else
+					{
+
+						const TManagedArray<FVector3f>* RenderVertex = &Flesh->Vertex;
+						if (GetDynamicCollection())
+						{
+							const TManagedArray<FVector3f>& DynamicVertex = GetDynamicCollection()->GetPositions();
+							if (DynamicVertex.Num()) RenderVertex = &DynamicVertex;
+						}
+						auto InRange = [](int32 Size, int32 Val) { return 0 <= Val && Val < Size; };
+
+						// Display only
+						for (int i = 0; i < NumFaces; ++i)
+						{
+							const auto& P1 = (*RenderVertex)[Flesh->Indices[i][0]];
+							const auto& P2 = (*RenderVertex)[Flesh->Indices[i][1]];
+							const auto& P3 = (*RenderVertex)[Flesh->Indices[i][2]];
+
+							RenderMesh->Vertices[3 * i] = FVector(P1);
+							RenderMesh->Vertices[3 * i + 1] = FVector(P2);
+							RenderMesh->Vertices[3 * i + 2] = FVector(P3);
+
+							auto Normal = Chaos::FVec3::CrossProduct(P3 - P1, P2 - P1);
+							RenderMesh->Normals[3 * i] = Normal;
+							RenderMesh->Normals[3 * i + 1] = Normal;
+							RenderMesh->Normals[3 * i + 2] = Normal;
+
+							auto Tangent = (P2 - P1).GetSafeNormal();
+							RenderMesh->Tangents[3 * i] = FProcMeshTangent(Tangent[0], Tangent[1], Tangent[2]);
+							Tangent = (P3 - P2).GetSafeNormal();
+							RenderMesh->Tangents[3 * i + 1] = FProcMeshTangent(Tangent[0], Tangent[1], Tangent[2]);
+							Tangent = (P1 - P3).GetSafeNormal();
+							RenderMesh->Tangents[3 * i + 2] = FProcMeshTangent(Tangent[0], Tangent[1], Tangent[2]);
+						}
+
+						if (!Mesh->GetComponentTransform().Equals(GetComponentTransform())) {
+							Mesh->SetRelativeTransform(GetComponentTransform());
+						}
+						Mesh->UpdateMeshSection_LinearColor(0, RenderMesh->Vertices, RenderMesh->Normals, RenderMesh->UVs, RenderMesh->Colors, RenderMesh->Tangents);
+					}
+
+					bCanRender = true;
+				}
+			}
+		}
+	}
+	if (!bCanRender)
+	{
+		ResetProceduralMesh();
+	}
+}
+
+
+
+
+
+
+
+
diff --git a/Engine/Plugins/Experimental/ChaosFlesh/Source/ChaosFleshEngine/Public/ChaosFlesh/ChaosDeformableSolverAsset.h b/Engine/Plugins/Experimental/ChaosFlesh/Source/ChaosFleshEngine/Public/ChaosFlesh/ChaosDeformableSolverAsset.h
new file mode 100644
--- /dev/null
+++ b/Engine/Plugins/Experimental/ChaosFlesh/Source/ChaosFleshEngine/Public/ChaosFlesh/ChaosDeformableSolverAsset.h
@@ -0,0 +1,27 @@
+// Copyright Epic Games, Inc. All Rights Reserved.
+
+#pragma once
+
+#include "CoreMinimal.h"
+#include "Physics/Experimental/PhysScene_Chaos.h"
+
+#include "ChaosDeformableSolverAsset.generated.h"
+
+class UChaosDeformableSolver;
+
+/**
+* UChaosDeformableSolver (UObject)
+*
+*/
+UCLASS(customconstructor)
+class CHAOSFLESHENGINE_API UChaosDeformableSolver : public UObject
+{
+	GENERATED_UCLASS_BODY()
+
+public:
+	UChaosDeformableSolver(const FObjectInitializer& ObjectInitializer = FObjectInitializer::Get());
+	bool IsVisible() { return true; }
+};
+
+
+
diff --git a/Engine/Plugins/Experimental/ChaosFlesh/Source/ChaosFleshEngine/Public/ChaosFlesh/ChaosDeformableSolverComponent.h b/Engine/Plugins/Experimental/ChaosFlesh/Source/ChaosFleshEngine/Public/ChaosFlesh/ChaosDeformableSolverComponent.h
new file mode 100644
--- /dev/null
+++ b/Engine/Plugins/Experimental/ChaosFlesh/Source/ChaosFleshEngine/Public/ChaosFlesh/ChaosDeformableSolverComponent.h
@@ -0,0 +1,107 @@
+// Copyright Epic Games, Inc. All Rights Reserved.
+#pragma once
+
+#include "Chaos/Deformable/ChaosDeformableSolver.h"
+#include "Chaos/Deformable/ChaosDeformableSolverProxy.h"
+#include "Chaos/Deformable/ChaosDeformableSolverTypes.h"
+#include "ChaosFlesh/ChaosDeformableSolverThreading.h"
+#include "ChaosFlesh/FleshComponent.h"
+#include "Components/SceneComponent.h"
+#include "UObject/ObjectMacros.h"
+
+#include "ChaosDeformableSolverComponent.generated.h"
+
+class UDeformablePhysicsComponent;
+
+/**
+*	UDeformableSolverComponent
+*/
+UCLASS()
+class CHAOSFLESHENGINE_API UDeformableSolverComponent : public USceneComponent
+{
+	GENERATED_UCLASS_BODY()
+
+public:
+	typedef Chaos::Softs::FThreadingProxy FThreadingProxy;
+	typedef Chaos::Softs::FFleshThreadingProxy FFleshThreadingProxy;
+	typedef Chaos::Softs::FDeformablePackage FDeformablePackage;
+	typedef Chaos::Softs::FDataMapValue FDataMapValue;
+	typedef Chaos::Softs::FDeformableSolver FDeformableSolver;
+
+	~UDeformableSolverComponent();
+
+#if WITH_EDITOR
+	TArray< TObjectPtr<AActor> > PreEditChangeDeformableActors;
+	virtual void PreEditChange(FProperty* PropertyThatWillChange) override;
+	virtual void PostEditChangeProperty(struct FPropertyChangedEvent& PropertyChangedEvent) override;
+#endif
+
+	/* Solver API */
+	FDeformableSolver::FGameThreadAccess GameThreadAccess();
+
+	bool IsSimulating(UDeformablePhysicsComponent*) const;
+	bool IsSimulatable() const;
+	void Reset();
+	void AddDeformableProxy(UDeformablePhysicsComponent* InComponent);
+	void Simulate(float DeltaTime);
+	void UpdateFromGameThread(float DeltaTime);
+	void UpdateFromSimulation(float DeltaTime);
+
+	/* Component Thread Management */
+	virtual void BeginPlay() override;
+	virtual void EndPlay(const EEndPlayReason::Type EndPlayReason) override;
+	virtual void TickComponent(float DeltaTime, enum ELevelTick TickType, FActorComponentTickFunction* ThisTickFunction) override;
+	//bool ShouldWaitForDeformableInTickFunction() const;
+	void UpdateDeformableEndTickState(bool bRegister);
+
+	/* Properties */
+
+	UPROPERTY(EditAnywhere, BlueprintReadOnly, Category = Chaos)
+	TArray< TObjectPtr<AActor> > DeformableActors;
+
+	UPROPERTY(EditAnywhere, Category = Chaos)
+		int32 NumSubSteps = 2;
+
+	UPROPERTY(EditAnywhere, Category = Chaos)
+		int32 NumSolverIterations = 5;
+
+	UPROPERTY(EditAnywhere, Category = Chaos)
+		bool FixTimeStep = false;
+
+	UPROPERTY(EditAnywhere, Category = Chaos)
+		float TimeStepSize = 0.05;
+
+	UPROPERTY(EditAnywhere, Category = Chaos)
+		bool CacheToFile = false;
+
+	UPROPERTY(EditAnywhere, Category = Chaos)
+		bool bEnableKinematics = true;
+
+	UPROPERTY(EditAnywhere, Category = Chaos)
+		bool bUseFloor = true;
+
+	UPROPERTY(EditAnywhere, Category = Chaos)
+		bool bDoSelfCollision = false;
+
+	UPROPERTY(EditAnywhere, Category = Chaos)
+		bool bDoThreadedAdvance = true;
+
+	UPROPERTY(EditAnywhere, Category = "Chaos|Expermental" )
+		bool bUseGridBasedConstraints = false;
+
+	UPROPERTY(EditAnywhere, Category = Chaos)
+		float GridDx = 25.;
+
+	//UPROPERTY(EditAnywhere, Category = Chaos)
+	//	bool bWaitForParallelDeformableTask = true;
+
+	// Simulation Variables
+	TUniquePtr<FDeformableSolver> Solver;
+
+protected:
+
+	/** Ref for the deformable solvers parallel task, so we can detect whether or not a sim is running */
+	FGraphEventRef ParallelDeformableTask;
+	FDeformableEndTickFunction DeformableEndTickFunction;
+};
+
diff --git a/Engine/Plugins/Experimental/ChaosFlesh/Source/ChaosFleshEngine/Public/ChaosFlesh/FleshComponent.h b/Engine/Plugins/Experimental/ChaosFlesh/Source/ChaosFleshEngine/Public/ChaosFlesh/FleshComponent.h
new file mode 100644
--- /dev/null
+++ b/Engine/Plugins/Experimental/ChaosFlesh/Source/ChaosFleshEngine/Public/ChaosFlesh/FleshComponent.h
@@ -0,0 +1,86 @@
+// Copyright Epic Games, Inc. All Rights Reserved.
+#pragma once
+
+#include "Chaos/Deformable/ChaosDeformableSolverProxy.h"
+#include "Chaos/Deformable/ChaosDeformableSolver.h"
+#include "ChaosFlesh/ChaosDeformableSolverThreading.h"
+#include "ChaosFlesh/ChaosDeformablePhysicsComponent.h"
+#include "ChaosFlesh/FleshAsset.h"
+#include "ChaosFlesh/FleshDynamicAsset.h"
+#include "Components/MeshComponent.h"
+#include "UObject/ObjectMacros.h"
+#include "ProceduralMeshComponent.h"
+#include "FleshComponent.generated.h"
+
+class FFleshCollection;
+class ADeformableSolverActor;
+class UDeformableSolverComponent;
+
+/**
+*	FleshComponent
+*/
+UCLASS(meta = (BlueprintSpawnableComponent))
+class CHAOSFLESHENGINE_API UFleshComponent : public UDeformablePhysicsComponent
+{
+	GENERATED_UCLASS_BODY()
+
+public:
+	typedef Chaos::Softs::FFleshThreadingProxy FFleshThreadingProxy;
+
+	~UFleshComponent();
+
+	virtual void BeginPlay() override;
+	virtual void OnRegister() override;
+	virtual void TickComponent(float DeltaTime, enum ELevelTick TickType, FActorComponentTickFunction *ThisTickFunction) override;
+	virtual FBoxSphereBounds CalcBounds(const FTransform& LocalToWorld) const override;
+	void UpdateLocalBounds();
+	virtual void EndPlay(const EEndPlayReason::Type ReasonEnd) override;
+	void Invalidate();
+
+	/** Simulation Interface*/
+	virtual FThreadingProxy* NewProxy() override;
+	virtual FDataMapValue NewDeformableData() override;
+	virtual void UpdateFromSimualtion(const FDataMapValue* SimualtionBuffer) override;
+
+
+	/** RestCollection */
+	void SetRestCollection(const UFleshAsset * InRestCollection);
+	const UFleshAsset* GetRestCollection() const { return RestCollection; }
+
+
+	/** DynamicCollection */
+	void ResetDynamicCollection();
+	UFleshDynamicAsset* GetDynamicCollection() { return DynamicCollection; }
+	const UFleshDynamicAsset* GetDynamicCollection() const { return DynamicCollection; }
+
+
+	UPROPERTY(EditAnywhere, BlueprintReadWrite, Category = Attributes)
+	TObjectPtr<UProceduralMeshComponent> Mesh;
+
+private:
+	UPROPERTY(EditAnywhere, Category = "ChaosPhysics")
+	TObjectPtr<const UFleshAsset> RestCollection;
+
+	UPROPERTY(EditAnywhere, Category = "ChaosPhysics")
+	TObjectPtr< UFleshDynamicAsset> DynamicCollection;
+
+	//
+	// Render the Procedural Mesh
+	//
+	struct FFleshRenderMesh
+	{
+		TArray<FVector> Vertices;
+		TArray<int32> Triangles;
+		TArray<FVector> Normals;
+		TArray<FVector2D> UVs;
+		TArray<FLinearColor> Colors;
+		TArray<FProcMeshTangent> Tangents;
+	};
+	FFleshRenderMesh* RenderMesh = nullptr;
+	void RenderProceduralMesh();
+	void ResetProceduralMesh();
+
+	bool bBoundsNeedsUpdate = true;
+	FBoxSphereBounds BoundingBox = FBoxSphereBounds(ForceInitToZero);
+};
+
diff --git a/Engine/Plugins/Experimental/ChaosFlesh/Source/ChaosFleshNodes/Private/Dataflow/ChaosFleshCoreNodes.cpp b/Engine/Plugins/Experimental/ChaosFlesh/Source/ChaosFleshNodes/Private/Dataflow/ChaosFleshCoreNodes.cpp
new file mode 100644
--- /dev/null
+++ b/Engine/Plugins/Experimental/ChaosFlesh/Source/ChaosFleshNodes/Private/Dataflow/ChaosFleshCoreNodes.cpp
@@ -0,0 +1,23 @@
+// Copyright Epic Games, Inc. All Rights Reserved.
+
+#include "Dataflow/ChaosFleshCoreNodes.h"
+
+#include "Dataflow/DataflowInputOutput.h"
+#include "Dataflow/DataflowNodeParameters.h"
+#include "Dataflow/DataflowNodeFactory.h"
+
+#include "GeometryCollection/ManagedArray.h"
+#include "GeometryCollection/ManagedArrayCollection.h"
+
+
+namespace Dataflow
+{
+	
+void RegisterChaosFleshCoreNodes()
+{
+	//DATAFLOW_NODE_REGISTER_CREATION_FACTORY(FNewManagedArrayCollectionNode);
+	//DATAFLOW_NODE_REGISTER_CREATION_FACTORY(FAddAttributeNode);
+}
+
+}
+
