ChaosFlesh : Moved to Experimental

diff --git a/Engine/Plugins/Experimental/ChaosFlesh/Source/ChaosFleshNodes/Private/Dataflow/ChaosFleshEngineAssetNodes.cpp b/Engine/Plugins/Experimental/ChaosFlesh/Source/ChaosFleshNodes/Private/Dataflow/ChaosFleshEngineAssetNodes.cpp
new file mode 100644
--- /dev/null
+++ b/Engine/Plugins/Experimental/ChaosFlesh/Source/ChaosFleshNodes/Private/Dataflow/ChaosFleshEngineAssetNodes.cpp
@@ -0,0 +1,390 @@
+// Copyright Epic Games, Inc. All Rights Reserved.
+
+#include "Dataflow/ChaosFleshEngineAssetNodes.h"
+
+#include "ChaosFlesh/FleshAsset.h"
+#include "ChaosFlesh/FleshCollectionUtility.h"
+#include "Chaos/Utilities.h"
+#include "Chaos/UniformGrid.h"
+#include "Dataflow/DataflowObjectInterface.h"
+#include "Dataflow/DataflowEngine.h"
+#include "GeometryCollection/ManagedArrayCollection.h"
+
+#include UE_INLINE_GENERATED_CPP_BY_NAME(ChaosFleshEngineAssetNodes)
+
+DEFINE_LOG_CATEGORY(LogChaosFlesh);
+
+namespace Dataflow
+{
+	void RegisterChaosFleshEngineAssetNodes()
+	{
+		DATAFLOW_NODE_REGISTER_CREATION_FACTORY(FGetFleshAssetDataflowNode);
+		DATAFLOW_NODE_REGISTER_CREATION_FACTORY(FExampleFleshEditDataflowNode);
+		DATAFLOW_NODE_REGISTER_CREATION_FACTORY(FSetFleshAssetDataflowNode);
+		DATAFLOW_NODE_REGISTER_CREATION_FACTORY(FImportFleshDataflowNode);
+		DATAFLOW_NODE_REGISTER_CREATION_FACTORY(FConstructTetGridNode);
+		DATAFLOW_NODE_REGISTER_CREATION_FACTORY(FComputeFleshMassNode);
+		DATAFLOW_NODE_REGISTER_CREATION_FACTORY(FComputeFiberFieldNode);
+	}
+}
+
+void FGetFleshAssetDataflowNode::Evaluate(Dataflow::FContext& Context, const FDataflowOutput* Out) const
+{
+	if (Out->IsA<FManagedArrayCollection>(&Output))
+	{
+		if (const Dataflow::FEngineContext* EngineContext = Context.AsType<Dataflow::FEngineContext>())
+		{
+			if (UFleshAsset* FleshAsset = Cast<UFleshAsset>(EngineContext->Owner))
+			{
+				if (const FFleshCollection* AssetCollection = FleshAsset->GetCollection())
+				{
+					SetValue<FManagedArrayCollection>(Context, *AssetCollection, &Output);
+				}
+			}
+		}
+	}
+}
+
+void FExampleFleshEditDataflowNode::Evaluate(Dataflow::FContext& Context, const FDataflowOutput* Out) const 
+{
+	if (Out->IsA<FManagedArrayCollection>(&Collection))
+	{
+		FManagedArrayCollection InputCollection = GetValue<FManagedArrayCollection>(Context, &Collection);
+
+		if (TManagedArray<FVector3f>* Vertex = InputCollection.FindAttribute<FVector3f>("Vertex", "Vertices"))
+		{
+			for (int i = 0; i < Vertex->Num(); i++)
+			{
+				(*Vertex)[i][1] *= Scale;
+			}
+		}
+
+		SetValue<FManagedArrayCollection>(Context, InputCollection, &Collection);
+	}
+}
+
+void FSetFleshAssetDataflowNode::Evaluate(Dataflow::FContext& Context, const FDataflowOutput* Out) const
+{
+	if (Out == nullptr)
+	{
+		const FManagedArrayCollection& Collection = GetValue<FManagedArrayCollection>(Context, &Input);
+
+		if (const Dataflow::FEngineContext* EngineContext = Context.AsType<Dataflow::FEngineContext>())
+		{
+			if (UFleshAsset* FleshAsset = Cast<UFleshAsset>(EngineContext->Owner))
+			{
+				FleshAsset->SetCollection(Collection.NewCopy<FFleshCollection>());
+			}
+		}
+	}
+}
+
+void FImportFleshDataflowNode::Evaluate(Dataflow::FContext& Context, const FDataflowOutput* Out) const
+{
+	if (Out->IsA<FManagedArrayCollection>(&Collection))
+	{
+		FString TempFileName = Filename.FilePath;
+		TempFileName = TempFileName.Replace(TEXT("\\"), TEXT("/"));
+		if (TUniquePtr<FFleshCollection> FleshCollection = ChaosFlesh::ImportTetFromFile(TempFileName))
+		{
+			SetValue<FManagedArrayCollection>(Context, *FleshCollection.Release(), &Collection);
+		}
+		else
+		{
+			SetValue<FManagedArrayCollection>(Context, Collection, &Collection);
+		}
+	}
+}
+
+
+void FConstructTetGridNode::Evaluate(Dataflow::FContext& Context, const FDataflowOutput* Out) const
+{
+	if (Out->IsA<FManagedArrayCollection>(&Collection))
+	{
+		Chaos::TVector<int32, 3> Counts(GridCellCount[0], GridCellCount[1], GridCellCount[2]);
+
+		Chaos::TVector<double, 3> MinCorner = -.5 * GridDomain;
+		Chaos::TVector<double, 3> MaxCorner = .5 * GridDomain;
+		Chaos::TUniformGrid<double, 3> Grid(MinCorner, MaxCorner, Counts, 0);
+
+		TArray<FIntVector4> Tets;
+		TArray<FVector> X;
+		Chaos::Utilities::TetMeshFromGrid<double>(Grid, Tets, X);
+
+		UE_LOG(LogChaosFlesh, Display, TEXT("TetGrid generated %d points and %d tetrahedra."), X.Num(), Tets.Num());
+
+		TUniquePtr<FFleshCollection> FleshCollection(
+			FFleshCollection::NewFleshCollection(X, Tets, false, false, false));
+		if (FleshCollection)
+		{
+			SetValue<FManagedArrayCollection>(Context, *FleshCollection.Release(), &Collection);
+		}
+		else
+		{
+			SetValue<FManagedArrayCollection>(Context, Collection, &Collection);
+		}
+	}
+}
+
+
+template <class T> using MType = FManagedArrayCollection::TManagedType<T>;
+
+void FComputeFleshMassNode::Evaluate(Dataflow::FContext& Context, const FDataflowOutput* Out) const
+{
+	if (Out->IsA<FManagedArrayCollection>(&Collection))
+	{
+		FManagedArrayCollection InCollection = GetValue<FManagedArrayCollection>(Context, &Collection);
+
+		if (InCollection.HasAttributes({
+			MType< float >("Mass", FGeometryCollection::VerticesGroup),
+			MType< FIntVector4 >(FTetrahedralCollection::TetrahedronAttribute, FTetrahedralCollection::TetrahedralGroup),
+			MType< FVector3f >("Vertex", "Vertices"),
+			MType< TArray<int32> >(FTetrahedralCollection::IncidentElementsAttribute, FGeometryCollection::VerticesGroup),
+			MType< TArray<int32> >(FTetrahedralCollection::IncidentElementsLocalIndexAttribute, FGeometryCollection::VerticesGroup) }))
+		{
+
+			int32 VertsNum = InCollection.NumElements(FGeometryCollection::VerticesGroup);
+			int32 TetsNum = InCollection.NumElements(FTetrahedralCollection::TetrahedralGroup);
+			if (VertsNum && TetsNum)
+			{
+				TManagedArray<float>& Mass = InCollection.ModifyAttribute<float>("Mass", FGeometryCollection::VerticesGroup);
+				const TManagedArray<FIntVector4>& Tetrahedron = InCollection.GetAttribute<FIntVector4>(FTetrahedralCollection::TetrahedronAttribute, FTetrahedralCollection::TetrahedralGroup);
+				const TManagedArray<FVector3f>& Vertex = InCollection.GetAttribute<FVector3f>("Vertex", "Vertices");
+				const TManagedArray<TArray<int32>>& ElemVertexIndex = InCollection.GetAttribute<TArray<int32>>(FTetrahedralCollection::IncidentElementsAttribute, FGeometryCollection::VerticesGroup);
+				const TManagedArray<TArray<int32>>& ElemLocalIndex = InCollection.GetAttribute<TArray<int32>>(FTetrahedralCollection::IncidentElementsLocalIndexAttribute, FGeometryCollection::VerticesGroup);
+
+				double Volume = 0.0;
+				TArray<float> ElementMass, Measure;
+				ElementMass.Init(0.f, 4 * TetsNum);
+				Measure.Init(0.f, TetsNum);
+				for (int e = 0; e < TetsNum; e++)
+				{
+					FVector3f X0 = Vertex[Tetrahedron[e][0]];
+					FVector3f X1 = Vertex[Tetrahedron[e][1]];
+					FVector3f X2 = Vertex[Tetrahedron[e][2]];
+					FVector3f X3 = Vertex[Tetrahedron[e][3]];
+					Measure[e] = ((X1 - X0).Dot(FVector3f::CrossProduct(X3 - X0, X2 - X0))) / 6.f;
+					if (Measure[e] < 0.f)
+					{
+						Measure[e] = -Measure[e];
+					}
+					Volume += Measure[e];
+				}
+				for (int32 e = 0; e < TetsNum; e++)
+				{
+					float ElementNodalMass = Density * Measure[e] / 4.f;
+					for (int i = 0; i < 4; i++)
+					{
+						ElementMass[4 * e + i] = ElementNodalMass;
+					}
+				}
+
+				bool WasMassSet = false;
+				for (int32 i = 0; i < ElemVertexIndex.Num(); i++)
+				{
+					if (ensure(ElemVertexIndex[i].Num() == ElemLocalIndex[i].Num()))
+					{
+						for (int32 j = 0; j < ElemVertexIndex[i].Num(); j++)
+						{
+							int32 TetIndex = ElemVertexIndex[i][j];
+							if (0 <= TetIndex && TetIndex < TetsNum)
+							{
+								int32 MassIndex = Tetrahedron[TetIndex][ElemLocalIndex[i][j]];
+								if (0 <= MassIndex && MassIndex < VertsNum)
+								{
+									WasMassSet = true;
+									Mass[MassIndex] += ElementMass[ElemVertexIndex[i][j] * 4 + ElemLocalIndex[i][j]];
+								}
+							}
+						}
+					}
+				}
+				if (!WasMassSet && Vertex.Num())
+				{
+					Mass.Fill(Density * Volume / Vertex.Num());
+				}
+			}
+		}
+		SetValue<FManagedArrayCollection>(Context, InCollection, &Collection);
+	}
+}
+
+void FComputeFiberFieldNode::Evaluate(Dataflow::FContext& Context, const FDataflowOutput* Out) const
+{
+	if (Out->IsA<FManagedArrayCollection>(&Collection))
+	{
+		//
+		// Gather inputs
+		//
+
+		FManagedArrayCollection InCollection = GetValue<FManagedArrayCollection>(Context, &Collection);
+		TArray<int32> InOriginIndices;// = GetValue<TArray<int32>>(Context, &OriginIndices);
+		TArray<int32> InInsertionIndices;// = GetValue<TArray<int32>>(Context, &InsertionIndices);
+
+		// Tetrahedra
+		TManagedArray<FIntVector4>* Elements = InCollection.FindAttribute<FIntVector4>(
+			FTetrahedralCollection::TetrahedronAttribute, FTetrahedralCollection::TetrahedralGroup);
+		if (!Elements)
+		{
+			UE_LOG(LogChaosFlesh, Warning,
+				TEXT("ComputeFiberFieldNode: Failed to find geometry collection attr '%s' in group '%s'"),
+				*FTetrahedralCollection::TetrahedronAttribute.ToString(), *FTetrahedralCollection::TetrahedralGroup.ToString());
+			Out->SetValue<FManagedArrayCollection>(InCollection, Context);
+			return;
+		}
+
+		// Vertices
+		TManagedArray<FVector3f>* Vertex = InCollection.FindAttribute<FVector3f>("Vertex", "Vertices");
+		if (!Vertex)
+		{
+			UE_LOG(LogChaosFlesh, Warning,
+				TEXT("ComputeFiberFieldNode: Failed to find geometry collection attr 'Vertex' in group 'Vertices'"));
+			Out->SetValue<FManagedArrayCollection>(InCollection, Context);
+			return;
+		}
+
+		// Incident elements
+		TManagedArray<TArray<int32>>* IncidentElements = InCollection.FindAttribute<TArray<int32>>(
+			FTetrahedralCollection::IncidentElementsAttribute, FGeometryCollection::VerticesGroup);
+		if (!IncidentElements)
+		{
+			UE_LOG(LogChaosFlesh, Warning,
+				TEXT("ComputeFiberFieldNode: Failed to find geometry collection attr '%s' in group '%s'"),
+				*FTetrahedralCollection::IncidentElementsAttribute.ToString(), *FGeometryCollection::VerticesGroup.ToString());
+			Out->SetValue<FManagedArrayCollection>(InCollection, Context);
+			return;
+		}
+		TManagedArray<TArray<int32>>* IncidentElementsLocalIndex = InCollection.FindAttribute<TArray<int32>>(
+			FTetrahedralCollection::IncidentElementsLocalIndexAttribute, FGeometryCollection::VerticesGroup);
+		if (!IncidentElementsLocalIndex)
+		{
+			UE_LOG(LogChaosFlesh, Warning,
+				TEXT("ComputeFiberFieldNode: Failed to find geometry collection attr '%s' in group '%s'"),
+				*FTetrahedralCollection::IncidentElementsLocalIndexAttribute.ToString(), *FGeometryCollection::VerticesGroup.ToString());
+			Out->SetValue<FManagedArrayCollection>(InCollection, Context);
+			return;
+		}
+
+		//
+		// Pull Origin & Insertion data out of the geometry collection.  We may want other ways of specifying
+		// these via an input on the node...
+		//
+
+		// Origin & Insertion
+		TManagedArray<int32>* Origin = nullptr; 
+		TManagedArray<int32>* Insertion = nullptr;
+		if (InOriginIndices.IsEmpty() || InInsertionIndices.IsEmpty())
+		{
+			// Origin & Insertion group
+			if (OriginInsertionGroupName.IsEmpty())
+			{
+				UE_LOG(LogChaosFlesh, Warning, TEXT("ComputeFiberFieldNode: Attr 'OriginInsertionGroupName' cannot be empty."));
+				Out->SetValue<FManagedArrayCollection>(InCollection, Context);
+				return;
+			}
+
+			// Origin vertices
+			if (InOriginIndices.IsEmpty())
+			{
+				if (OriginVertexFieldName.IsEmpty())
+				{
+					UE_LOG(LogChaosFlesh, Warning, TEXT("ComputeFiberFieldNode: Attr 'OriginVertexFieldName' cannot be empty."));
+					Out->SetValue<FManagedArrayCollection>(InCollection, Context);
+					return;
+				}
+				Origin = InCollection.FindAttribute<int32>(FName(OriginVertexFieldName), FName(OriginInsertionGroupName));
+				if (!Origin)
+				{
+					UE_LOG(LogChaosFlesh, Warning,
+						TEXT("ComputeFiberFieldNode: Failed to find geometry collection attr '%s' in group '%s'"),
+						*OriginVertexFieldName, *OriginInsertionGroupName);
+					Out->SetValue<FManagedArrayCollection>(InCollection, Context);
+					return;
+				}
+			}
+
+			// Insertion vertices
+			if (InInsertionIndices.IsEmpty())
+			{
+				if (InsertionVertexFieldName.IsEmpty())
+				{
+					UE_LOG(LogChaosFlesh, Warning, TEXT("ComputeFiberFieldNode: Attr 'InsertionVertexFieldName' cannot be empty."));
+					Out->SetValue<FManagedArrayCollection>(InCollection, Context);
+					return;
+				}
+				Insertion = InCollection.FindAttribute<int32>(FName(InsertionVertexFieldName), FName(OriginInsertionGroupName));
+				if (!Insertion)
+				{
+					UE_LOG(LogChaosFlesh, Warning,
+						TEXT("ComputeFiberFieldNode: Failed to find geometry collection attr '%s' in group '%s'"),
+						*InsertionVertexFieldName, *OriginInsertionGroupName);
+					Out->SetValue<FManagedArrayCollection>(InCollection, Context);
+					return;
+				}
+			}
+		}
+
+		//
+		// Do the thing
+		//
+
+		TArray<FVector3f> FiberDirs =
+			ComputeFiberField(*Elements, *Vertex, *IncidentElements, *IncidentElementsLocalIndex, 
+				Origin ? Origin->GetConstArray() : InOriginIndices,
+				Insertion ? Insertion->GetConstArray() : InInsertionIndices);
+
+		//
+		// Set output(s)
+		//
+
+		TManagedArray<FVector3f>* FiberDirections =
+			InCollection.FindAttribute<FVector3f>("FiberDirection", FTetrahedralCollection::TetrahedralGroup);
+		if (!FiberDirections)
+		{
+			FiberDirections =
+				&InCollection.AddAttribute<FVector3f>("FiberDirection", FTetrahedralCollection::TetrahedralGroup);
+		}
+		(*FiberDirections) = MoveTemp(FiberDirs);
+
+		Out->SetValue<FManagedArrayCollection>(InCollection, Context);
+	}
+}
+
+
+TArray<int32> 
+FComputeFiberFieldNode::GetNonZeroIndices(const TArray<uint8>& Map) const
+{
+	int32 NumNonZero = 0;
+	for (int32 i = 0; i < Map.Num(); i++)
+		if (Map[i])
+			NumNonZero++;
+	TArray<int32> Indices; Indices.AddUninitialized(NumNonZero);
+	int32 Idx = 0;
+	for (int32 i = 0; i < Map.Num(); i++)
+		if (Map[i])
+			Indices[Idx++] = i;
+	return Indices;
+}
+
+TArray<FVector3f>
+FComputeFiberFieldNode::ComputeFiberField(
+	const TManagedArray<FIntVector4>& Elements,
+	const TManagedArray<FVector3f>& Vertex,
+	const TManagedArray<TArray<int32>>& IncidentElements,
+	const TManagedArray<TArray<int32>>& IncidentElementsLocalIndex,
+	const TArray<int32>& Origin,
+	const TArray<int32>& Insertion) const
+{
+	TArray<FVector3f> Directions;
+	Chaos::ComputeFiberField<float>(
+		Elements.GetConstArray(),
+		Vertex.GetConstArray(),
+		IncidentElements.GetConstArray(),
+		IncidentElementsLocalIndex.GetConstArray(),
+		Origin,
+		Insertion,
+		Directions,
+		MaxIterations,
+		Tolerance);
+	return Directions;
+}
diff --git a/Engine/Plugins/Experimental/ChaosFlesh/Source/ChaosFleshNodes/Private/Dataflow/ChaosFleshFiberDirectionInitializationNodes.cpp b/Engine/Plugins/Experimental/ChaosFlesh/Source/ChaosFleshNodes/Private/Dataflow/ChaosFleshFiberDirectionInitializationNodes.cpp
new file mode 100644
--- /dev/null
+++ b/Engine/Plugins/Experimental/ChaosFlesh/Source/ChaosFleshNodes/Private/Dataflow/ChaosFleshFiberDirectionInitializationNodes.cpp
@@ -0,0 +1,28 @@
+// Copyright Epic Games, Inc. All Rights Reserved.
+
+#include "Dataflow/ChaosFleshFiberDirectionInitializationNodes.h"
+
+#include "GeometryCollection/ManagedArrayCollection.h"
+
+namespace Dataflow
+{
+	void ChaosFleshFiberDirectionInitializationNodes()
+	{
+		DATAFLOW_NODE_REGISTER_CREATION_FACTORY(FGenerateFiberDirectionsDataflowNode);
+	}
+}
+
+void FGenerateFiberDirectionsDataflowNode::Evaluate(Dataflow::FContext& Context, const FDataflowOutput* Out) const
+{
+	if (Out->IsA<DataType>(&Collection))
+	{
+		DataType InCollection = GetValue<DataType>(Context, &Collection);
+
+		//
+		// @todo(dataflow) : Implemention fiber direction
+		//
+
+		SetValue<DataType>(Context, InCollection, &Collection);
+	}
+}
+
diff --git a/Engine/Plugins/Experimental/ChaosFlesh/Source/ChaosFleshNodes/Private/Dataflow/ChaosFleshKinematicInitializationNodes.cpp b/Engine/Plugins/Experimental/ChaosFlesh/Source/ChaosFleshNodes/Private/Dataflow/ChaosFleshKinematicInitializationNodes.cpp
new file mode 100644
--- /dev/null
+++ b/Engine/Plugins/Experimental/ChaosFlesh/Source/ChaosFleshNodes/Private/Dataflow/ChaosFleshKinematicInitializationNodes.cpp
@@ -0,0 +1,288 @@
+// Copyright Epic Games, Inc. All Rights Reserved.
+
+#include "Dataflow/ChaosFleshKinematicInitializationNodes.h"
+
+#include "ChaosFlesh/FleshAsset.h"
+#include "ChaosFlesh/FleshCollectionUtility.h"
+#include "Chaos/Convex.h"
+#include "Dataflow/DataflowCore.h"
+#include "Dataflow/DataflowEngine.h"
+#include "Engine/SkeletalMesh.h"
+#include "GeometryCollection/Facades/CollectionKinematicBindingFacade.h"
+#include "GeometryCollection/Facades/CollectionVertexBoneWeightsFacade.h"
+#include "GeometryCollection/ManagedArrayCollection.h"
+#include "GeometryCollection/ManagedArrayTypes.h"
+
+#include UE_INLINE_GENERATED_CPP_BY_NAME(ChaosFleshKinematicInitializationNodes)
+
+//DEFINE_LOG_CATEGORY_STATIC(FKinematicInitializationNodesLog, Log, All);
+
+namespace Dataflow
+{
+	void RegisterChaosFleshKinematicInitializationNodes()
+	{
+		DATAFLOW_NODE_REGISTER_CREATION_FACTORY(FKinematicInitializationDataflowNode);
+		DATAFLOW_NODE_REGISTER_CREATION_FACTORY(FSelectionSetDataflowNode);
+		DATAFLOW_NODE_REGISTER_CREATION_FACTORY(FKinematicTetrahedralBindingsDataflowNode);
+		DATAFLOW_NODE_REGISTER_CREATION_FACTORY(FSetVerticesKinematicDataflowNode);
+		DATAFLOW_NODE_REGISTER_CREATION_FACTORY(FBinVerticesDataflowNode);
+	}
+}
+
+void FKinematicTetrahedralBindingsDataflowNode::Evaluate(Dataflow::FContext& Context, const FDataflowOutput* Out) const
+{
+	if (Out->IsA<DataType>(&Collection))
+	{
+		DataType InCollection = GetValue<DataType>(Context, &Collection);
+		TManagedArray<FIntVector4>* Tetrahedron = InCollection.FindAttribute<FIntVector4>(FTetrahedralCollection::TetrahedronAttribute, FTetrahedralCollection::TetrahedralGroup);
+		TManagedArray<FVector3f>* Vertex = InCollection.FindAttribute<FVector3f>("Vertex", "Vertices");
+
+		TObjectPtr<USkeletalMesh> SkeletalMesh = GetValue<TObjectPtr<USkeletalMesh>>(Context, &SkeletalMeshIn);
+		if (SkeletalMesh && Tetrahedron && Vertex)
+		{
+			//parse exclusion list to find bones to skip
+			TArray<FString> StrArray;
+			ExclusionList.ParseIntoArray(StrArray, *FString(" "));				
+
+			int32 NumTets = Tetrahedron->Num();
+			TArray<FTransform> ComponentPose;
+			Dataflow::Animation::GlobalTransforms(SkeletalMesh->GetRefSkeleton(), ComponentPose);
+			for (int32 b = 0; b < SkeletalMesh->GetRefSkeleton().GetNum(); b++)
+			{
+				bool Skip = false;
+				FString BoneName = SkeletalMesh->GetRefSkeleton().GetBoneName(b).ToString();
+				for (FString Elem : StrArray)
+				{
+					if (BoneName.Contains(Elem))
+					{
+						Skip = true;
+						break;
+					}
+				}
+				if (Skip)
+					continue;
+
+				FVector3f BonePosition(ComponentPose[b].GetTranslation());
+				int32 ParentIndex=SkeletalMesh->GetRefSkeleton().GetParentIndex(b);
+				
+				if (ParentIndex != INDEX_NONE) 
+				{
+					FVector3f ParentPosition(ComponentPose[ParentIndex].GetTranslation());
+					FVector3f RayDir = ParentPosition - BonePosition;
+					Chaos::FReal Length = RayDir.Length();
+					RayDir.Normalize();
+
+					if (Length > Chaos::FReal(1e-8)) 
+					{
+						TSet<int32> BoneVertSet;
+						for (int32 t = 0; t < NumTets; t++) 
+						{
+							int32 i = (*Tetrahedron)[t][0];
+							int32 j = (*Tetrahedron)[t][1];
+							int32 k = (*Tetrahedron)[t][2];
+							int32 l = (*Tetrahedron)[t][3];
+
+							TArray<Chaos::TVec3<Chaos::FRealSingle>> InVertices;
+							InVertices.SetNum(4);
+							InVertices[0][0] = (*Vertex)[i].X; InVertices[0][1] = (*Vertex)[i].Y; InVertices[0][2] = (*Vertex)[i].Z;
+							InVertices[1][0] = (*Vertex)[j].X; InVertices[1][1] = (*Vertex)[j].Y; InVertices[1][2] = (*Vertex)[j].Z;
+							InVertices[2][0] = (*Vertex)[k].X; InVertices[2][1] = (*Vertex)[k].Y; InVertices[2][2] = (*Vertex)[k].Z;
+							InVertices[3][0] = (*Vertex)[l].X; InVertices[3][1] = (*Vertex)[l].Y; InVertices[3][2] = (*Vertex)[l].Z;
+							Chaos::FConvex ConvexTet(InVertices, Chaos::FReal(0));
+							Chaos::FReal OutTime;
+							Chaos::FVec3 OutPosition, OutNormal;
+							int32 OutFaceIndex;
+							bool KeepTet = ConvexTet.Raycast(BonePosition, RayDir, Length, Chaos::FReal(0), OutTime, OutPosition, OutNormal, OutFaceIndex);
+							if (KeepTet) 
+							{
+								BoneVertSet.Add(i); BoneVertSet.Add(j); BoneVertSet.Add(k); BoneVertSet.Add(l);
+							}
+						}
+
+						TArray<int32> BoundVerts = BoneVertSet.Array();
+						TArray<float> BoundWeights;
+						BoundWeights.Init(float(1), BoundVerts.Num());
+						
+						//get local coords of bound verts
+						typedef Chaos::Facades::FKinematicBindingFacade Kinematics;
+						Kinematics::FBindingKey Binding = Kinematics::SetBoneBindings(&InCollection, b, BoundVerts, BoundWeights);
+						TManagedArray<TArray<FVector3f>>& LocalPos = InCollection.AddAttribute<TArray<FVector3f>>("LocalPosition", Binding.GroupName);
+						Kinematics::AddKinematicBinding(&InCollection, Binding);
+
+						auto FloatVert = [](FVector3d V) { return FVector3f(V.X, V.Y, V.Z); };
+						auto DoubleVert = [](FVector3f V) { return FVector3d(V.X, V.Y, V.Z); };
+						LocalPos[Binding.Index].SetNum(BoundVerts.Num());
+						for (int32 i = 0; i < BoundVerts.Num(); i++)
+						{	
+							FVector3f Temp = (*Vertex)[BoundVerts[i]];
+							LocalPos[Binding.Index][i] = FloatVert(ComponentPose[b].InverseTransformPosition(DoubleVert(Temp)));
+						}
+					}
+				}
+			}
+			GeometryCollection::Facades::FVertexBoneWeightsFacade::AddBoneWeightsFromKinematicBindings(&InCollection);
+		}
+
+		SetValue<DataType>(Context, InCollection, &Collection);
+	}
+}
+
+
+void FKinematicInitializationDataflowNode::Evaluate(Dataflow::FContext& Context, const FDataflowOutput* Out) const
+{
+	if (Out->IsA<DataType>(&Collection))
+	{
+		DataType InCollection = GetValue<DataType>(Context, &Collection);
+
+		if (TObjectPtr<USkeletalMesh> SkeletalMesh = GetValue<TObjectPtr<USkeletalMesh>>(Context, &SkeletalMeshIn))
+		{
+			int32 Index = GetValue<int32>(Context, &BoneIndexIn);
+			if (Index != INDEX_NONE)
+			{
+				if (TManagedArray<FVector3f>* Vertices = InCollection.FindAttribute<FVector3f>("Vertex", FGeometryCollection::VerticesGroup))
+				{
+					TArray<int32> BoundVerts;
+					TArray<float> BoundWeights;
+					if (FindInput(&VertexIndicesIn) && FindInput(&VertexIndicesIn)->GetConnection())
+					{
+						for (int32 SelectionIndex : GetValue<TArray<int32>>(Context, &VertexIndicesIn))
+						{
+							if (0 <= SelectionIndex && SelectionIndex < Vertices->Num())
+							{
+								BoundVerts.Add(SelectionIndex);
+							}
+						}
+						BoundWeights.Init(1.0, BoundVerts.Num());
+					}
+					else
+					{
+						TArray<FTransform> ComponentPose;
+						Dataflow::Animation::GlobalTransforms(SkeletalMesh->GetRefSkeleton(), ComponentPose);
+						if (0 < Index && Index < ComponentPose.Num())
+						{
+							FVector3f BonePosition(ComponentPose[Index].GetTranslation());
+
+							int NumVertices = Vertices->Num();
+							for (int i = Vertices->Num() - 1; i > 0; i--)
+							{
+								if ((BonePosition - (*Vertices)[i]).Length() < Radius)
+								{
+									BoundVerts.Add(i);
+									BoundWeights.Add(1.0);
+								}
+							}
+						}
+					}
+					typedef Chaos::Facades::FKinematicBindingFacade Kinematics;
+					Kinematics::AddKinematicBinding(&InCollection, Kinematics::SetBoneBindings(&InCollection, Index, BoundVerts, BoundWeights));
+				}
+			}
+		}
+		SetValue<DataType>(Context, InCollection, &Collection);
+	}
+}
+
+void FSetVerticesKinematicDataflowNode::Evaluate(Dataflow::FContext& Context, const FDataflowOutput* Out) const
+{
+	if (Out->IsA<DataType>(&Collection))
+	{
+		DataType InCollection = GetValue<DataType>(Context, &Collection);
+		TArray<int32> BoundVerts;
+		TArray<float> BoundWeights;
+		if (FindInput(&VertexIndicesIn) && FindInput(&VertexIndicesIn)->GetConnection())
+		{
+			if (TManagedArray<FVector3f>* Vertices = InCollection.FindAttribute<FVector3f>("Vertex", FGeometryCollection::VerticesGroup))
+			{
+
+				for (int32 SelectionIndex : GetValue<TArray<int32>>(Context, &VertexIndicesIn))
+				{
+					if (0 <= SelectionIndex && SelectionIndex < Vertices->Num())
+					{
+						BoundVerts.Add(SelectionIndex);
+					}
+				}
+				BoundWeights.Init(1.0, BoundVerts.Num());	
+			}
+		} 
+		else
+		{
+			if (TManagedArray<FVector3f>* Vertices = InCollection.FindAttribute<FVector3f>("Vertex", FGeometryCollection::VerticesGroup))
+			{
+				float MaxZ = -FLT_MAX;
+				int MaxIndex = INDEX_NONE;
+				for (int i = 0; i < Vertices->Num(); i++)
+				{
+					if ((*Vertices)[i].Z > MaxZ)
+					{
+						MaxZ = (*Vertices)[i].Z;
+						MaxIndex = i;
+					}
+				}
+				if (MaxIndex != INDEX_NONE)
+				{
+					BoundVerts.Add(MaxIndex);
+					BoundWeights.Add(1.f);
+				}
+			}
+			if (BoundVerts.Num() > 0)
+			{
+				typedef Chaos::Facades::FKinematicBindingFacade Kinematics;
+				Kinematics::AddKinematicBinding(&InCollection, Kinematics::SetBoneBindings(&InCollection, INDEX_NONE, BoundVerts, BoundWeights));
+			}	
+		}
+		SetValue<DataType>(Context, InCollection, &Collection);
+	}
+}
+
+
+void FSelectionSetDataflowNode::Evaluate(Dataflow::FContext& Context, const FDataflowOutput* Out) const
+{
+	if (Out->IsA<DataType>(&IndicesOut))
+	{
+		TArray<int32> IntArray;
+		TArray<FString> StrArray;
+		Indices.ParseIntoArray(StrArray, *FString(" "));
+		for (FString Elem : StrArray)
+		{
+			if (FCString::IsNumeric(*Elem))
+			{
+				IntArray.Add(FCString::Atoi(*Elem));
+			}
+		}
+
+		Out->SetValue<DataType>(IntArray, Context);
+	}
+}
+
+
+void FBinVerticesDataflowNode::Evaluate(Dataflow::FContext& Context, const FDataflowOutput* Out) const
+{
+	if (Out->IsA<DataType>(&VertexIndicesOut))
+	{	
+		TArray<int32> VertexArray;
+		FManagedArrayCollection InCollection = GetValue<FManagedArrayCollection>(Context, &Collection);
+		FString TempFileName = Filename;
+		TempFileName = TempFileName.Replace(TEXT("\\"), TEXT("/"));
+		if (TUniquePtr<FFleshCollection> FleshCollection = ChaosFlesh::ImportTetFromFile(TempFileName))
+		{	
+			TManagedArray<FVector3f>* VerticesBin = FleshCollection->FindAttribute<FVector3f>("Vertex", FGeometryCollection::VerticesGroup);
+			TManagedArray<FVector3f>* VerticesCollection = InCollection.FindAttribute<FVector3f>("Vertex", FGeometryCollection::VerticesGroup);
+			if (VerticesBin && VerticesCollection)
+			{
+				for (int32 i = 0; i < VerticesBin->Num(); i++)
+				{
+					for (int32 j = 0; j < VerticesCollection->Num(); ++j)
+					{
+						if (((*VerticesBin)[i] - (*VerticesCollection)[j]).Length() < Tolerance)
+						{
+							VertexArray.Emplace(j);
+							break;
+						}
+					}
+					ensureMsgf(VertexArray.Num() == i + 1, TEXT("Some vertex is not binned or binned repeatedly at current tolerance"));
+				}
+			}
+		}
+		Out->SetValue<DataType>(VertexArray, Context);
+	}
+}
diff --git a/Engine/Plugins/Experimental/ChaosFlesh/Source/ChaosFleshNodes/Private/Dataflow/ChaosFleshNodesPlugin.cpp b/Engine/Plugins/Experimental/ChaosFlesh/Source/ChaosFleshNodes/Private/Dataflow/ChaosFleshNodesPlugin.cpp
new file mode 100644
--- /dev/null
+++ b/Engine/Plugins/Experimental/ChaosFlesh/Source/ChaosFleshNodes/Private/Dataflow/ChaosFleshNodesPlugin.cpp
@@ -0,0 +1,37 @@
+// Copyright Epic Games, Inc. All Rights Reserved.
+
+#include "Dataflow/ChaosFleshNodesPlugin.h"
+
+#include "CoreMinimal.h"
+
+#include "Dataflow/DataflowNodeFactory.h"
+#include "Dataflow/ChaosFleshEngineAssetNodes.h"
+#include "Dataflow/ChaosFleshCoreNodes.h"
+#include "Dataflow/ChaosFleshFiberDirectionInitializationNodes.h"
+#include "Dataflow/ChaosFleshKinematicInitializationNodes.h"
+#include "Dataflow/ChaosFleshRenderInitializationNodes.h"
+#include "Dataflow/ChaosFleshTetrahedralNodes.h"
+
+
+#define LOCTEXT_NAMESPACE "ChaosFleshNodes"
+
+
+void IChaosFleshNodesPlugin::StartupModule()
+{
+	Dataflow::RegisterChaosFleshEngineAssetNodes();
+	Dataflow::RegisterChaosFleshCoreNodes();
+	Dataflow::ChaosFleshFiberDirectionInitializationNodes();
+	Dataflow::ChaosFleshRenderInitializationNodes();
+	Dataflow::RegisterChaosFleshKinematicInitializationNodes();
+	Dataflow::ChaosFleshTetrahedralNodes();
+}
+
+void IChaosFleshNodesPlugin::ShutdownModule()
+{
+}
+
+
+IMPLEMENT_MODULE(IChaosFleshNodesPlugin, ChaosFleshNodes)
+
+
+#undef LOCTEXT_NAMESPACE
diff --git a/Engine/Plugins/Experimental/ChaosFlesh/Source/ChaosFleshNodes/Private/Dataflow/ChaosFleshRenderInitializationNodes.cpp b/Engine/Plugins/Experimental/ChaosFlesh/Source/ChaosFleshNodes/Private/Dataflow/ChaosFleshRenderInitializationNodes.cpp
new file mode 100644
--- /dev/null
+++ b/Engine/Plugins/Experimental/ChaosFlesh/Source/ChaosFleshNodes/Private/Dataflow/ChaosFleshRenderInitializationNodes.cpp
@@ -0,0 +1,31 @@
+// Copyright Epic Games, Inc. All Rights Reserved.
+
+#include "Dataflow/ChaosFleshRenderInitializationNodes.h"
+
+#include "GeometryCollection/ManagedArrayCollection.h"
+
+
+
+namespace Dataflow
+{
+	void ChaosFleshRenderInitializationNodes()
+	{
+		DATAFLOW_NODE_REGISTER_CREATION_FACTORY(FBindForRenderToSkeletalMeshDataflowNode);
+	}
+}
+
+void FBindForRenderToSkeletalMeshDataflowNode::Evaluate(Dataflow::FContext& Context, const FDataflowOutput* Out) const
+{
+	if (Out->IsA<DataType>(&Collection))
+	{
+		DataType InCollection = GetValue<DataType>(Context, &Collection);
+
+		//
+		// @todo(dataflow) : Implemention binding for optimus deformation pipeline. 
+		//
+
+		SetValue<DataType>(Context, InCollection, &Collection);
+	}
+}
+
+
diff --git a/Engine/Plugins/Experimental/ChaosFlesh/Source/ChaosFleshNodes/Private/Dataflow/ChaosFleshTetrahedralNodes.cpp b/Engine/Plugins/Experimental/ChaosFlesh/Source/ChaosFleshNodes/Private/Dataflow/ChaosFleshTetrahedralNodes.cpp
new file mode 100644
--- /dev/null
+++ b/Engine/Plugins/Experimental/ChaosFlesh/Source/ChaosFleshNodes/Private/Dataflow/ChaosFleshTetrahedralNodes.cpp
@@ -0,0 +1,120 @@
+// Copyright Epic Games, Inc. All Rights Reserved.
+
+#include "Dataflow/ChaosFleshTetrahedralNodes.h"
+
+#include "Chaos/Deformable/Utilities.h"
+#include "ChaosFlesh/FleshCollection.h"
+#include "CompGeom/ExactPredicates.h"
+#include "DynamicMesh/DynamicMesh3.h"
+#include "DynamicMesh/DynamicMeshAABBTree3.h"
+#include "Engine/StaticMesh.h"
+#include "Generate/IsosurfaceStuffing.h"
+#include "GeometryCollection/ManagedArrayCollection.h"
+#include "MeshDescriptionToDynamicMesh.h"
+#include "Spatial/FastWinding.h"
+
+namespace Dataflow
+{
+	void ChaosFleshTetrahedralNodes()
+	{
+		DATAFLOW_NODE_REGISTER_CREATION_FACTORY(FGenerateTetrahedralCollectionDataflowNodes);
+	}
+}
+
+// Helper to get the boundary of a tet mesh, useful for debugging / verifying output
+TArray<FIntVector3> GetSurfaceTriangles(const TArray<FIntVector4>& Tets)
+{
+	// Rotate the vector so the first element is the smallest
+	auto RotVec = [](const FIntVector3& F) -> FIntVector3
+	{
+		int32 MinIdx = F.X < F.Y ? (F.X < F.Z ? 0 : 2) : (F.Y < F.Z ? 1 : 2);
+		return FIntVector3(F[MinIdx], F[(MinIdx + 1) % 3], F[(MinIdx + 2) % 3]);
+	};
+	// Reverse the winding while keeping the first element unchanged
+	auto RevVec = [](const FIntVector3& F) -> FIntVector3
+	{
+		return FIntVector3(F.X, F.Z, F.Y);
+	};
+
+	TSet<FIntVector3> FacesSet;
+	for (int32 TetIdx = 0; TetIdx < Tets.Num(); ++TetIdx)
+	{
+		FIntVector3 TetF[4];
+		Chaos::Utilities::GetTetFaces(Tets[TetIdx], TetF[0], TetF[1], TetF[2], TetF[3], false);
+		for (int32 SubIdx = 0; SubIdx < 4; ++SubIdx)
+		{
+			FIntVector3 Key = RotVec(TetF[SubIdx]);
+			if (FacesSet.Contains(Key))
+			{
+				FacesSet.Remove(Key);
+			}
+			else
+			{
+				FacesSet.Add(RevVec(Key));
+			}
+		}
+	}
+	return FacesSet.Array();
+}
+
+void FGenerateTetrahedralCollectionDataflowNodes::Evaluate(Dataflow::FContext& Context, const FDataflowOutput* Out) const
+{
+	if (Out->IsA<DataType>(&Collection))
+	{
+		TUniquePtr<FFleshCollection> InCollection(GetValue<DataType>(Context, &Collection).NewCopy<FFleshCollection>());
+
+#if WITH_EDITORONLY_DATA
+		if (StaticMesh && NumCells>0 && (-.5<= OffsetPercent && OffsetPercent<=0.5))
+		{
+			// make a mesh description for UE::Geometry tools
+			UE::Geometry::FDynamicMesh3 DynamicMesh;
+			FMeshDescriptionToDynamicMesh GetSourceMesh;
+			bool bUsingHiResSource = StaticMesh->IsHiResMeshDescriptionValid();
+			const FMeshDescription* UseSourceMeshDescription =
+				(bUsingHiResSource) ? StaticMesh->GetHiResMeshDescription() : StaticMesh->GetMeshDescription(0);
+			GetSourceMesh.Convert(UseSourceMeshDescription, DynamicMesh);
+
+
+			// Tet mesh generation
+			UE::Geometry::TIsosurfaceStuffing<double> IsosurfaceStuffing;
+			UE::Geometry::FDynamicMeshAABBTree3 Spatial(&DynamicMesh);
+			UE::Geometry::TFastWindingTree<UE::Geometry::FDynamicMesh3> FastWinding(&Spatial);
+			UE::Geometry::FAxisAlignedBox3d Bounds = Spatial.GetBoundingBox();
+			IsosurfaceStuffing.Bounds = FBox(Bounds);
+			double CellSize = Bounds.MaxDim() / NumCells;
+			IsosurfaceStuffing.CellSize = CellSize;
+			IsosurfaceStuffing.IsoValue = .5+OffsetPercent;
+			IsosurfaceStuffing.Implicit = [&FastWinding, &Spatial](FVector3d Pos)
+			{
+				FVector3d Nearest = Spatial.FindNearestPoint(Pos);
+				double WindingSign = FastWinding.FastWindingNumber(Pos) - .5;
+				return FVector3d::Distance(Nearest, Pos) * FMathd::SignNonZero(WindingSign);
+			};
+
+			IsosurfaceStuffing.Generate();
+			if (IsosurfaceStuffing.Tets.Num() > 0)
+			{
+				TArray<FVector> Vertices; Vertices.SetNumUninitialized(IsosurfaceStuffing.Vertices.Num());
+				TArray<FIntVector4> Elements; Elements.SetNumUninitialized(IsosurfaceStuffing.Tets.Num());
+				TArray<FIntVector3> SurfaceElements = GetSurfaceTriangles(IsosurfaceStuffing.Tets);
+
+				for (int32 Tdx = 0; Tdx < IsosurfaceStuffing.Tets.Num(); ++Tdx)
+				{
+					Elements[Tdx] = IsosurfaceStuffing.Tets[Tdx];
+				}
+				for (int32 Vdx = 0; Vdx < IsosurfaceStuffing.Vertices.Num(); ++Vdx)
+				{
+					Vertices[Vdx] = IsosurfaceStuffing.Vertices[Vdx];
+				}
+
+				TUniquePtr<FTetrahedralCollection> TetCollection(FTetrahedralCollection::NewTetrahedralCollection(Vertices, SurfaceElements, Elements));
+				InCollection->AppendGeometry(*TetCollection.Get());
+			}
+		}
+#else
+		ensureMsgf(false,TEXT("FGenerateTetrahedralCollectionDataflowNodes is a editor only node."));
+#endif
+
+		SetValue<DataType>(Context, *(FManagedArrayCollection*)InCollection.Get(), &Collection);
+	}
+}
diff --git a/Engine/Plugins/Experimental/ChaosFlesh/Source/ChaosFleshNodes/Public/Dataflow/ChaosFleshCoreNodes.h b/Engine/Plugins/Experimental/ChaosFlesh/Source/ChaosFleshNodes/Public/Dataflow/ChaosFleshCoreNodes.h
new file mode 100644
--- /dev/null
+++ b/Engine/Plugins/Experimental/ChaosFlesh/Source/ChaosFleshNodes/Public/Dataflow/ChaosFleshCoreNodes.h
@@ -0,0 +1,127 @@
+// Copyright Epic Games, Inc. All Rights Reserved.
+
+#pragma once
+
+#include "CoreMinimal.h"
+
+#include "Dataflow/DataflowInputOutput.h"
+#include "Dataflow/DataflowNode.h"
+#include "GeometryCollection/ManagedArray.h"
+#include "GeometryCollection/ManagedArrayCollection.h"
+
+namespace Dataflow
+{
+	/**
+	typedef TSharedPtr<FManagedArrayCollection> FManagedArrayCollectionSharedPtr;
+
+	enum class EManagedArrayType : uint8
+	{
+		EgManagedArrayTypeNone = 0,
+		EgManagedArrayTypeBool,
+		EgManagedArrayTypeInt,
+		EgManagedArrayTypeFloat,
+		EgManagedArrayTypeVector
+	};
+
+	class CHAOSFLESHCORENODES_API FNewManagedArrayCollectionNode : public FNode
+	{
+		DATAFLOW_NODE_DEFINE_INTERNAL(FNewManagedArrayCollectionNode)
+
+	public:
+		typedef TSharedPtr<FManagedArrayCollection> DataType;
+		DataType Value;
+		TSharedPtr< class TOutput<DataType> > Output;
+
+		FNewManagedArrayCollectionNode(const FNodeParameters& InParam, FGuid InGuid = FGuid::NewGuid())
+			: FNode(InParam, InGuid)
+			, Value(new FManagedArrayCollection())
+			, Output(new TOutput<DataType>(TOutputParameters<DataType>({FName("CollectionOut"), this})))
+		{}
+
+		virtual void Evaluate(const FContext& Context, FConnection* Out) const override
+		{
+			if (Output.Get() == Out)
+			{
+				Output->SetValue(Value, Context);
+			}
+		}
+	};
+
+	class CHAOSFLESHCORENODES_API FAddAttributeNode : public FNode
+	{
+		DATAFLOW_NODE_DEFINE_INTERNAL(FAddAttributeNode)
+
+	public:
+
+		typedef TSharedPtr<FManagedArrayCollection> DataType;
+
+		FName AttributeName;
+		FName GroupName;
+		EManagedArrayType AttributeType;
+
+		TSharedPtr< class TInput<DataType> > Input;
+		TSharedPtr< class TInput<int32> > SizeInput;
+		TSharedPtr< class TOutput<DataType> > Output;
+
+		FAddAttributeNode(const FNodeParameters& InParam, FGuid InGuid = FGuid::NewGuid())
+			: FNode(InParam, InGuid)
+			, AttributeName("Position")
+			, GroupName("Particle")
+			, AttributeType(EManagedArrayType::EgManagedArrayTypeVector)
+			, Input(new TInput<DataType>(TInputParameters<DataType>({FName("CollectionIn"), this})))
+			, SizeInput(new TInput<int32>(TInputParameters<int32>({FName("SizeIn"), this })))
+			, Output(new TOutput<DataType>(TOutputParameters<DataType>({FName("CollectionOut"), this })))
+		{}
+
+		virtual void Evaluate(const FContext& Context, FConnection* Out) const override
+		{
+			if (Output.Get() == Out)
+			{
+				DataType Collection(new FManagedArrayCollection());
+
+				// its an error to not set the requested output
+				if (DataType Val = Input->GetValue(Context))
+				{
+					Collection = Val;
+
+					int NumElements = 100;
+					if (SizeInput->GetConnection())
+					{
+						NumElements = SizeInput->GetValue(Context);
+					}
+
+					if (!Collection->HasGroup(GroupName))
+					{
+						Collection->AddGroup(GroupName);
+					}
+					if (!Collection->HasAttribute(AttributeName, GroupName))
+					{
+						Collection->AddAttribute<FVector>(AttributeName, GroupName);
+					}
+
+					if (Collection->NumElements(GroupName) < NumElements)
+					{
+						Collection->AddElements(NumElements, GroupName);
+
+						TManagedArray<FVector>& Pos = Collection->ModifyAttribute<FVector>(AttributeName, GroupName);
+						for (int i = 0; i < NumElements; i++)
+						{
+							Pos[i] = FVector(i);
+						}
+					}
+				}
+				Output->SetValue(Collection, Context);
+			}
+		}
+
+		virtual void SerializeInternal(FArchive& Ar) override { 
+			Ar << AttributeName << GroupName << AttributeType;
+		};
+
+	};
+
+
+	*/
+	void RegisterChaosFleshCoreNodes();
+}
+
diff --git a/Engine/Plugins/Experimental/ChaosFlesh/Source/ChaosFleshNodes/Public/Dataflow/ChaosFleshEngineAssetNodes.h b/Engine/Plugins/Experimental/ChaosFlesh/Source/ChaosFleshNodes/Public/Dataflow/ChaosFleshEngineAssetNodes.h
new file mode 100644
--- /dev/null
+++ b/Engine/Plugins/Experimental/ChaosFlesh/Source/ChaosFleshNodes/Public/Dataflow/ChaosFleshEngineAssetNodes.h
@@ -0,0 +1,228 @@
+// Copyright Epic Games, Inc. All Rights Reserved.
+
+#pragma once 
+
+#include "CoreMinimal.h"
+#include "Dataflow/DataflowCore.h"
+#include "Chaos/Math/Poisson.h"
+#include "ChaosLog.h"
+#include "ChaosFlesh/TetrahedralCollection.h"
+#include "Dataflow/DataflowEngine.h"
+
+#include "ChaosFleshEngineAssetNodes.generated.h"
+
+USTRUCT()
+struct FGetFleshAssetDataflowNode : public FDataflowNode
+{
+	GENERATED_USTRUCT_BODY()
+	DATAFLOW_NODE_DEFINE_INTERNAL(FGetFleshAssetDataflowNode, "GetFleshAsset", "Flesh", "")
+	DATAFLOW_NODE_RENDER_TYPE(FGeometryCollection::StaticType(), "Output")
+
+public:
+
+	UPROPERTY(meta = (DataflowOutput, DisplayName = "Collection"))
+	FManagedArrayCollection Output;
+
+	FGetFleshAssetDataflowNode(const Dataflow::FNodeParameters& InParam, FGuid InGuid = FGuid::NewGuid())
+		: FDataflowNode(InParam, InGuid)
+	{
+		RegisterOutputConnection(&Output);
+	}
+
+	virtual void Evaluate(Dataflow::FContext& Context, const FDataflowOutput* Out) const override;
+};
+
+
+
+
+USTRUCT()
+struct FExampleFleshEditDataflowNode : public FDataflowNode
+{
+	GENERATED_USTRUCT_BODY()
+	DATAFLOW_NODE_DEFINE_INTERNAL(FExampleFleshEditDataflowNode, "ExampleFleshEdit", "Flesh", "")
+	DATAFLOW_NODE_RENDER_TYPE(FGeometryCollection::StaticType(), "Collection")
+
+public:
+
+	UPROPERTY(EditAnywhere, Category = "Dataflow")
+	float Scale = 10.0;
+
+	UPROPERTY(meta = (DataflowInput, DataflowOutput, DisplayName = "Collection", Passthrough = "Collection"))
+	FManagedArrayCollection Collection;
+
+	FExampleFleshEditDataflowNode(const Dataflow::FNodeParameters& InParam, FGuid InGuid = FGuid::NewGuid())
+		: FDataflowNode(InParam, InGuid)
+	{
+		RegisterInputConnection(&Collection);
+		RegisterOutputConnection(&Collection, &Collection);
+	}
+
+	virtual void Evaluate(Dataflow::FContext& Context, const FDataflowOutput* Out) const override;
+};
+
+
+USTRUCT()
+struct FSetFleshAssetDataflowNode : public FDataflowNode
+{
+	GENERATED_USTRUCT_BODY()
+	DATAFLOW_NODE_DEFINE_INTERNAL(FSetFleshAssetDataflowNode, "SetFleshAsset", "Flesh", "")
+
+public:
+
+	UPROPERTY(meta = (DataflowInput, DisplayName = "Collection"))
+	FManagedArrayCollection Input;
+
+	FSetFleshAssetDataflowNode(const Dataflow::FNodeParameters& InParam, FGuid InGuid = FGuid::NewGuid())
+		: FDataflowNode(InParam, InGuid)
+	{
+		RegisterInputConnection(&Input);
+	}
+
+	virtual void Evaluate(Dataflow::FContext& Context, const FDataflowOutput* Out) const override;
+};
+
+
+
+USTRUCT()
+struct FImportFleshDataflowNode : public FDataflowNode
+{
+	GENERATED_USTRUCT_BODY()
+	DATAFLOW_NODE_DEFINE_INTERNAL(FImportFleshDataflowNode, "ImportFlesh", "Flesh", "")
+	DATAFLOW_NODE_RENDER_TYPE(FGeometryCollection::StaticType(), "Collection")
+
+public:
+
+	UPROPERTY(EditAnywhere, Category = "Dataflow")
+	FFilePath Filename;
+
+	UPROPERTY(meta = (DataflowOutput, DisplayName = "Collection"))
+	FManagedArrayCollection Collection;
+	
+	FImportFleshDataflowNode(const Dataflow::FNodeParameters& InParam, FGuid InGuid = FGuid::NewGuid())
+		: FDataflowNode(InParam, InGuid)
+	{
+		RegisterOutputConnection(&Collection);
+	}
+
+	virtual void Evaluate(Dataflow::FContext& Context, const FDataflowOutput* Out) const override;
+};
+
+
+USTRUCT()
+struct FConstructTetGridNode : public FDataflowNode
+{
+	GENERATED_USTRUCT_BODY()
+	DATAFLOW_NODE_DEFINE_INTERNAL(FConstructTetGridNode, "TetGrid", "Flesh", "")
+	DATAFLOW_NODE_RENDER_TYPE(FGeometryCollection::StaticType(), "Collection")
+
+public:
+
+	UPROPERTY(meta = (DataflowOutput, DisplayName = "Collection"))
+	FManagedArrayCollection Collection;
+
+	UPROPERTY(EditAnywhere, Category = "Dataflow", meta = (ClampMin = "1"))
+	FIntVector GridCellCount = FIntVector(10, 10, 10);
+
+	UPROPERTY(EditAnywhere, Category = "Dataflow", meta = (ClampMin = "0.0"))
+	FVector GridDomain = FVector(1.0, 1.0, 1.0);
+
+	FConstructTetGridNode(const Dataflow::FNodeParameters& InParam, FGuid InGuid = FGuid::NewGuid())
+		: FDataflowNode(InParam, InGuid)
+	{
+		RegisterOutputConnection(&Collection);
+	}
+
+	virtual void Evaluate(Dataflow::FContext& Context, const FDataflowOutput* Out) const override;
+};
+
+
+
+USTRUCT()
+struct FComputeFleshMassNode : public FDataflowNode
+{
+	GENERATED_USTRUCT_BODY()
+	DATAFLOW_NODE_DEFINE_INTERNAL(FComputeFleshMassNode, "ComputeFleshMass", "Flesh", "")
+	DATAFLOW_NODE_RENDER_TYPE(FGeometryCollection::StaticType(), "Collection")
+
+public:
+
+	UPROPERTY(EditAnywhere, Category = "Dataflow")
+	float Density = 1.f;
+
+	UPROPERTY(meta = (DataflowInput, DataflowOutput, DisplayName = "Collection", Passthrough = "Collection"))
+	FManagedArrayCollection Collection;
+
+	FComputeFleshMassNode(const Dataflow::FNodeParameters& InParam, FGuid InGuid = FGuid::NewGuid())
+		: FDataflowNode(InParam, InGuid)
+	{
+		RegisterInputConnection(&Collection);
+		RegisterOutputConnection(&Collection, &Collection);
+	}
+
+	virtual void Evaluate(Dataflow::FContext& Context, const FDataflowOutput* Out) const override;
+};
+
+/**
+* Computes a muscle fiber direction per tetrahedron from a GeometryCollection containing tetrahedra, 
+* vertices, and origin & insertion vertex fields.  Fiber directions should smoothly follow the geometry
+* oriented from the origin vertices pointing to the insertion vertices.
+*/
+USTRUCT()
+struct FComputeFiberFieldNode : public FDataflowNode
+{
+	GENERATED_USTRUCT_BODY()
+	DATAFLOW_NODE_DEFINE_INTERNAL(FComputeFiberFieldNode, "ComputeFiberField", "Flesh", "")
+	DATAFLOW_NODE_RENDER_TYPE(FGeometryCollection::StaticType(), "Collection")
+
+public:
+	//typedef FManagedArrayCollection DataType;
+
+	UPROPERTY(meta = (DataflowInput, DataflowOutput, DisplayName = "Collection", Passthrough = "Collection"))
+	FManagedArrayCollection Collection;
+
+	//UPROPERTY(meta = (DataflowInput, DisplayName = "OriginVertexIndices"))
+	//TArray<int32> OriginIndices;
+
+	//UPROPERTY(meta = (DataflowInput, DisplayName = "InsertionVertexIndices"))
+	//TArray<int32> InsertionIndices;
+
+	UPROPERTY(EditAnywhere, Category = "Dataflow")
+	FString OriginInsertionGroupName = FString();
+
+	UPROPERTY(EditAnywhere, Category = "Dataflow")
+	FString OriginVertexFieldName = FString("Origin");
+
+	UPROPERTY(EditAnywhere, Category = "Dataflow")
+	FString InsertionVertexFieldName = FString("Insertion");
+
+	UPROPERTY(EditAnywhere, Category = "Dataflow")
+	int32 MaxIterations = 100;
+
+	UPROPERTY(EditAnywhere, Category = "Dataflow")
+	float Tolerance = 1.0e-7;
+
+	FComputeFiberFieldNode(const Dataflow::FNodeParameters& InParam, FGuid InGuid = FGuid::NewGuid())
+		: FDataflowNode(InParam, InGuid)
+	{
+		RegisterInputConnection(&Collection);
+		RegisterOutputConnection(&Collection, &Collection);
+	}
+
+	virtual void Evaluate(Dataflow::FContext& Context, const FDataflowOutput* Out) const override;
+
+	TArray<int32> GetNonZeroIndices(const TArray<uint8>& Map) const;
+
+	TArray<FVector3f> ComputeFiberField(
+		const TManagedArray<FIntVector4>& Elements,
+		const TManagedArray<FVector3f>& Vertex,
+		const TManagedArray<TArray<int32>>& IncidentElements,
+		const TManagedArray<TArray<int32>>& IncidentElementsLocalIndex,
+		const TArray<int32>& Origin,
+		const TArray<int32>& Insertion) const;
+};
+
+namespace Dataflow
+{
+	void RegisterChaosFleshEngineAssetNodes();
+}
+
diff --git a/Engine/Plugins/Experimental/ChaosFlesh/Source/ChaosFleshNodes/Public/Dataflow/ChaosFleshFiberDirectionInitializationNodes.h b/Engine/Plugins/Experimental/ChaosFlesh/Source/ChaosFleshNodes/Public/Dataflow/ChaosFleshFiberDirectionInitializationNodes.h
new file mode 100644
--- /dev/null
+++ b/Engine/Plugins/Experimental/ChaosFlesh/Source/ChaosFleshNodes/Public/Dataflow/ChaosFleshFiberDirectionInitializationNodes.h
@@ -0,0 +1,39 @@
+// Copyright Epic Games, Inc. All Rights Reserved.
+
+#pragma once 
+
+#include "CoreMinimal.h"
+#include "Dataflow/DataflowCore.h"
+#include "Dataflow/DataflowEngine.h"
+
+#include "ChaosFleshFiberDirectionInitializationNodes.generated.h"
+
+USTRUCT()
+struct FGenerateFiberDirectionsDataflowNode : public FDataflowNode
+{
+	GENERATED_USTRUCT_BODY()
+	DATAFLOW_NODE_DEFINE_INTERNAL(FGenerateFiberDirectionsDataflowNode, "GenerateFiberDirections", "Flesh", "")
+	DATAFLOW_NODE_RENDER_TYPE(FGeometryCollection::StaticType(), "Collection")
+
+public:
+	typedef FManagedArrayCollection DataType;
+
+	UPROPERTY(meta = (DataflowInput, DataflowOutput, DisplayName = "Collection", Passthrough = "Collection"))
+	FManagedArrayCollection Collection;
+
+	FGenerateFiberDirectionsDataflowNode(const Dataflow::FNodeParameters& InParam, FGuid InGuid = FGuid::NewGuid())
+		: FDataflowNode(InParam, InGuid)
+	{
+		RegisterInputConnection(&Collection);
+		RegisterOutputConnection(&Collection, &Collection);
+	}
+
+	virtual void Evaluate(Dataflow::FContext& Context, const FDataflowOutput* Out) const override;
+};
+
+
+namespace Dataflow
+{
+	void ChaosFleshFiberDirectionInitializationNodes();
+}
+
diff --git a/Engine/Plugins/Experimental/ChaosFlesh/Source/ChaosFleshNodes/Public/Dataflow/ChaosFleshKinematicInitializationNodes.h b/Engine/Plugins/Experimental/ChaosFlesh/Source/ChaosFleshNodes/Public/Dataflow/ChaosFleshKinematicInitializationNodes.h
new file mode 100644
--- /dev/null
+++ b/Engine/Plugins/Experimental/ChaosFlesh/Source/ChaosFleshNodes/Public/Dataflow/ChaosFleshKinematicInitializationNodes.h
@@ -0,0 +1,171 @@
+// Copyright Epic Games, Inc. All Rights Reserved.
+
+#pragma once 
+
+#include "CoreMinimal.h"
+#include "Dataflow/DataflowCore.h"
+#include "Dataflow/DataflowEngine.h"
+
+#include "ChaosFleshKinematicInitializationNodes.generated.h"
+
+class USkeletalMesh;
+
+USTRUCT()
+struct FKinematicTetrahedralBindingsDataflowNode : public FDataflowNode
+{
+	GENERATED_USTRUCT_BODY()
+	DATAFLOW_NODE_DEFINE_INTERNAL(FKinematicTetrahedralBindingsDataflowNode, "KinematicTetrahedralBindings", "Flesh", "")
+	DATAFLOW_NODE_RENDER_TYPE(FGeometryCollection::StaticType(), "Collection")
+
+public:
+	typedef FManagedArrayCollection DataType;
+
+	UPROPERTY(meta = (DataflowInput, DisplayName = "SkeletalMesh"))
+	TObjectPtr<USkeletalMesh> SkeletalMeshIn = nullptr;
+
+	UPROPERTY(meta = (DataflowInput, DataflowOutput, DisplayName = "Collection", Passthrough = "Collection"))
+	FManagedArrayCollection Collection;
+
+	UPROPERTY(EditAnywhere, Category = "Dataflow")
+	FString ExclusionList = "twist foo";
+
+	FKinematicTetrahedralBindingsDataflowNode(const Dataflow::FNodeParameters& InParam, FGuid InGuid = FGuid::NewGuid())
+		: FDataflowNode(InParam, InGuid)
+	{
+		RegisterInputConnection(&SkeletalMeshIn);
+		RegisterInputConnection(&Collection);
+		RegisterOutputConnection(&Collection, &Collection);
+	}
+
+	virtual void Evaluate(Dataflow::FContext& Context, const FDataflowOutput* Out) const override;
+};
+
+
+USTRUCT()
+struct FKinematicInitializationDataflowNode : public FDataflowNode
+{
+	GENERATED_USTRUCT_BODY()
+	DATAFLOW_NODE_DEFINE_INTERNAL(FKinematicInitializationDataflowNode, "KinematicInitialization", "Flesh", "")
+	DATAFLOW_NODE_RENDER_TYPE(FGeometryCollection::StaticType(), "Collection")
+
+public:
+	typedef FManagedArrayCollection DataType;
+
+	UPROPERTY(EditAnywhere, Category = "Dataflow")
+	float Radius = 40.f;
+
+	UPROPERTY(EditAnywhere, Category = "Dataflow")
+	FTransform Transform;
+
+	UPROPERTY(meta = (DataflowInput, DataflowOutput, DisplayName = "Collection", Passthrough = "Collection"))
+	FManagedArrayCollection Collection;
+
+	UPROPERTY(meta = (DataflowInput, DisplayName = "SkeletalMesh"))
+	TObjectPtr<USkeletalMesh> SkeletalMeshIn;
+		
+	UPROPERTY(meta = (DataflowInput, DisplayName = "SelectionSet"))
+	TArray<int32> VertexIndicesIn;
+
+	UPROPERTY(meta = (DataflowInput, DisplayName = "BoneIndex"))
+	int32 BoneIndexIn;
+	
+	FKinematicInitializationDataflowNode(const Dataflow::FNodeParameters& InParam, FGuid InGuid = FGuid::NewGuid())
+		: FDataflowNode(InParam, InGuid)
+	{
+		RegisterInputConnection(&Collection);
+		RegisterOutputConnection(&Collection, &Collection);
+		RegisterInputConnection(&SkeletalMeshIn);
+		RegisterInputConnection(&VertexIndicesIn);	
+		RegisterInputConnection(&BoneIndexIn);
+	}
+
+	virtual void Evaluate(Dataflow::FContext& Context, const FDataflowOutput* Out) const override;
+};
+
+
+
+USTRUCT()
+struct FSelectionSetDataflowNode : public FDataflowNode
+{
+	GENERATED_USTRUCT_BODY()
+	DATAFLOW_NODE_DEFINE_INTERNAL(FSelectionSetDataflowNode, "SelectionSet", "Flesh", "")
+
+public:
+	typedef TArray<int32> DataType;
+
+	UPROPERTY(EditAnywhere, Category = "Dataflow")
+	FString Indices = FString("1 2 3");
+
+	UPROPERTY(meta = (DataflowOutput, DisplayName = "SelectionSet"))
+	TArray<int32> IndicesOut;
+
+	FSelectionSetDataflowNode(const Dataflow::FNodeParameters& InParam, FGuid InGuid = FGuid::NewGuid())
+		: FDataflowNode(InParam, InGuid)
+	{
+		RegisterOutputConnection(&IndicesOut);
+	}
+
+	virtual void Evaluate(Dataflow::FContext& Context, const FDataflowOutput* Out) const override;
+};
+
+
+
+USTRUCT()
+struct FSetVerticesKinematicDataflowNode : public FDataflowNode
+{
+	GENERATED_USTRUCT_BODY()
+	DATAFLOW_NODE_DEFINE_INTERNAL(FSetVerticesKinematicDataflowNode, "SetVerticesKinematic", "Flesh", "")
+
+public:
+	typedef FManagedArrayCollection DataType;
+
+	UPROPERTY(meta = (DataflowInput, DataflowOutput, DisplayName = "Collection", Passthrough = "Collection"))
+		FManagedArrayCollection Collection;
+
+	UPROPERTY(meta = (DataflowInput, DisplayName = "SelectionSet"))
+		TArray<int32> VertexIndicesIn;
+
+	FSetVerticesKinematicDataflowNode(const Dataflow::FNodeParameters& InParam, FGuid InGuid = FGuid::NewGuid())
+		: FDataflowNode(InParam, InGuid)
+	{
+		RegisterInputConnection(&Collection);
+		RegisterOutputConnection(&Collection, &Collection);
+		RegisterInputConnection(&VertexIndicesIn);
+	}
+
+	virtual void Evaluate(Dataflow::FContext& Context, const FDataflowOutput* Out) const override;
+};
+
+USTRUCT()
+struct FBinVerticesDataflowNode : public FDataflowNode
+{
+	GENERATED_USTRUCT_BODY()
+	DATAFLOW_NODE_DEFINE_INTERNAL(FBinVerticesDataflowNode, "BinVertices", "Flesh", "")
+	DATAFLOW_NODE_RENDER_TYPE(FGeometryCollection::StaticType(), "Collection")
+
+public:
+	typedef TArray<int32> DataType;
+	UPROPERTY(EditAnywhere, Category = "Dataflow")
+		FString Filename = FString("D:/UE5/Main/QAGame/Import/example.geo");
+	UPROPERTY(EditAnywhere, Category = "Dataflow")
+		float Tolerance = float(1e-6);
+	UPROPERTY(meta = (DataflowInput, DisplayName = "Collection"))
+		FManagedArrayCollection Collection;
+	UPROPERTY(meta = (DataflowOutput, DisplayName = "SelectionSet"))
+		TArray<int32> VertexIndicesOut;
+
+	FBinVerticesDataflowNode(const Dataflow::FNodeParameters& InParam, FGuid InGuid = FGuid::NewGuid())
+		: FDataflowNode(InParam, InGuid)
+	{
+		RegisterInputConnection(&Collection);
+		RegisterOutputConnection(&VertexIndicesOut);
+	}
+
+	virtual void Evaluate(Dataflow::FContext& Context, const FDataflowOutput* Out) const override;
+};
+
+namespace Dataflow
+{
+	void RegisterChaosFleshKinematicInitializationNodes();
+}
+
diff --git a/Engine/Plugins/Experimental/ChaosFlesh/Source/ChaosFleshNodes/Public/Dataflow/ChaosFleshRenderInitializationNodes.h b/Engine/Plugins/Experimental/ChaosFlesh/Source/ChaosFleshNodes/Public/Dataflow/ChaosFleshRenderInitializationNodes.h
new file mode 100644
--- /dev/null
+++ b/Engine/Plugins/Experimental/ChaosFlesh/Source/ChaosFleshNodes/Public/Dataflow/ChaosFleshRenderInitializationNodes.h
@@ -0,0 +1,44 @@
+// Copyright Epic Games, Inc. All Rights Reserved.
+
+#pragma once 
+
+#include "CoreMinimal.h"
+#include "Dataflow/DataflowCore.h"
+#include "Dataflow/DataflowEngine.h"
+
+#include "ChaosFleshRenderInitializationNodes.generated.h"
+
+class USkeletalMesh;
+
+
+USTRUCT()
+struct FBindForRenderToSkeletalMeshDataflowNode : public FDataflowNode
+{
+	GENERATED_USTRUCT_BODY()
+		DATAFLOW_NODE_DEFINE_INTERNAL(FBindForRenderToSkeletalMeshDataflowNode, "BindForRenderToSkeletalMesh", "Flesh", "")
+
+public:
+	typedef FManagedArrayCollection DataType;
+
+	UPROPERTY(meta = (DataflowInput, DisplayName = "SkeletalMesh"))
+	TObjectPtr<const USkeletalMesh> SkeletalMesh = nullptr;
+
+	UPROPERTY(meta = (DataflowInput, DataflowOutput, DisplayName = "Collection", Passthrough = "Collection"))
+		FManagedArrayCollection Collection;
+
+	FBindForRenderToSkeletalMeshDataflowNode(const Dataflow::FNodeParameters& InParam, FGuid InGuid = FGuid::NewGuid())
+		: FDataflowNode(InParam, InGuid)
+	{
+		RegisterInputConnection(&SkeletalMesh);
+		RegisterInputConnection(&Collection);
+		RegisterOutputConnection(&Collection, &Collection);
+	}
+
+	virtual void Evaluate(Dataflow::FContext& Context, const FDataflowOutput* Out) const override;
+};
+
+namespace Dataflow
+{
+	void ChaosFleshRenderInitializationNodes();
+}
+
diff --git a/Engine/Plugins/Experimental/ChaosFlesh/Source/ChaosFleshNodes/Public/Dataflow/ChaosFleshTetrahedralNodes.h b/Engine/Plugins/Experimental/ChaosFlesh/Source/ChaosFleshNodes/Public/Dataflow/ChaosFleshTetrahedralNodes.h
new file mode 100644
--- /dev/null
+++ b/Engine/Plugins/Experimental/ChaosFlesh/Source/ChaosFleshNodes/Public/Dataflow/ChaosFleshTetrahedralNodes.h
@@ -0,0 +1,58 @@
+// Copyright Epic Games, Inc. All Rights Reserved.
+
+#pragma once 
+
+#include "CoreMinimal.h"
+#include "Dataflow/DataflowEngine.h"
+#include "GeometryCollection/ManagedArrayCollection.h"
+
+#include "ChaosFleshTetrahedralNodes.generated.h"
+
+class UStaticMesh;
+
+
+USTRUCT()
+struct FGenerateTetrahedralCollectionDataflowNodes : public FDataflowNode
+{
+	GENERATED_USTRUCT_BODY()
+	DATAFLOW_NODE_DEFINE_INTERNAL(FGenerateTetrahedralCollectionDataflowNodes, "GenerateTetrahedralCollection", "Flesh", "")
+	DATAFLOW_NODE_RENDER_TYPE(FGeometryCollection::StaticType(), "Collection")
+
+public:
+	typedef FManagedArrayCollection DataType;
+
+	UPROPERTY(EditAnywhere, Category = "Dataflow")
+	int32 NumCells = 32;
+
+	UPROPERTY(EditAnywhere, Category = "Dataflow")
+	double OffsetPercent = 0.05;
+
+	UPROPERTY(EditAnywhere, Category = "Dataflow")
+	TObjectPtr<const UStaticMesh> StaticMesh = nullptr;
+
+	UPROPERTY(meta = (DataflowInput, DataflowOutput, DisplayName = "Collection"))
+	FManagedArrayCollection Collection;
+
+	FGenerateTetrahedralCollectionDataflowNodes(const Dataflow::FNodeParameters& InParam, FGuid InGuid = FGuid::NewGuid())
+		: FDataflowNode(InParam, InGuid)
+	{
+		RegisterInputConnection(&StaticMesh);
+		RegisterInputConnection(&Collection);
+		RegisterOutputConnection(&Collection);
+	}
+
+	virtual void Evaluate(Dataflow::FContext& Context, const FDataflowOutput* Out) const override;
+
+};
+
+namespace Dataflow
+{
+	void ChaosFleshTetrahedralNodes();
+
+
+}
+
+
+
+
+
