Fixed non-unity/pch errors in UnrealGame linux config

diff --git a/Engine/Plugins/Enterprise/VariantManagerContent/Source/VariantManagerContent/Private/ThumbnailGenerator.cpp b/Engine/Plugins/Enterprise/VariantManagerContent/Source/VariantManagerContent/Private/ThumbnailGenerator.cpp
--- a/Engine/Plugins/Enterprise/VariantManagerContent/Source/VariantManagerContent/Private/ThumbnailGenerator.cpp
+++ b/Engine/Plugins/Enterprise/VariantManagerContent/Source/VariantManagerContent/Private/ThumbnailGenerator.cpp
@@ -2,6 +2,7 @@
 
 #include "ThumbnailGenerator.h"
 
+#include "Engine/World.h"
 #include "Engine/World.h"
 #include "VariantManagerContentLog.h"
 
diff --git a/Engine/Plugins/Experimental/ChaosFlesh/Source/ChaosFleshEngine/Private/ChaosFlesh/FleshComponent.cpp b/Engine/Plugins/Experimental/ChaosFlesh/Source/ChaosFleshEngine/Private/ChaosFlesh/FleshComponent.cpp
--- a/Engine/Plugins/Experimental/ChaosFlesh/Source/ChaosFleshEngine/Private/ChaosFlesh/FleshComponent.cpp
+++ b/Engine/Plugins/Experimental/ChaosFlesh/Source/ChaosFleshEngine/Private/ChaosFlesh/FleshComponent.cpp
@@ -3,6 +3,7 @@
 #include "ChaosFlesh/FleshComponent.h"
 
 #include "Animation/SkeletalMeshActor.h"
+#include "Animation/Skeleton.h"
 #include "Chaos/Deformable/ChaosDeformableSolverProxy.h"
 #include "ChaosFlesh/FleshCollectionUtility.h"
 #include "ChaosFlesh/ChaosDeformableSolverActor.h"
diff --git a/Engine/Plugins/Experimental/CommonUI/Source/CommonInput/Private/CommonInputSubsystem.cpp b/Engine/Plugins/Experimental/CommonUI/Source/CommonInput/Private/CommonInputSubsystem.cpp
--- a/Engine/Plugins/Experimental/CommonUI/Source/CommonInput/Private/CommonInputSubsystem.cpp
+++ b/Engine/Plugins/Experimental/CommonUI/Source/CommonInput/Private/CommonInputSubsystem.cpp
@@ -2,6 +2,7 @@
 
 #include "CommonInputSubsystem.h"
 #include "CommonInputPrivate.h"
+#include "Framework/Application/SlateApplication.h"
 #include "Misc/ConfigCacheIni.h"
 
 #include "Framework/Application/IInputProcessor.h"
diff --git a/Engine/Plugins/Experimental/Enterprise/DatasmithRuntime/Source/Private/SceneImporter.cpp b/Engine/Plugins/Experimental/Enterprise/DatasmithRuntime/Source/Private/SceneImporter.cpp
--- a/Engine/Plugins/Experimental/Enterprise/DatasmithRuntime/Source/Private/SceneImporter.cpp
+++ b/Engine/Plugins/Experimental/Enterprise/DatasmithRuntime/Source/Private/SceneImporter.cpp
@@ -13,6 +13,7 @@
 
 #include "Async/Async.h"
 #include "Materials/MaterialInstanceDynamic.h"
+#include "Misc/Paths.h"
 #include "ProfilingDebugging/MiscTrace.h"
 #include "RenderingThread.h"
 #include "Misc/Paths.h"
diff --git a/Engine/Plugins/Experimental/Text3D/Source/Text3D/Private/MeshCreator.h b/Engine/Plugins/Experimental/Text3D/Source/Text3D/Private/MeshCreator.h
--- a/Engine/Plugins/Experimental/Text3D/Source/Text3D/Private/MeshCreator.h
+++ b/Engine/Plugins/Experimental/Text3D/Source/Text3D/Private/MeshCreator.h
@@ -7,6 +7,7 @@
 #include "Mesh.h"
 #include "ContourNode.h"
 
+class UStaticMesh;
 
 /** Makes actual bevel. */
 class FMeshCreator final
diff --git a/Engine/Plugins/Experimental/Text3D/Source/Text3D/Public/Glyph.h b/Engine/Plugins/Experimental/Text3D/Source/Text3D/Public/Glyph.h
--- a/Engine/Plugins/Experimental/Text3D/Source/Text3D/Public/Glyph.h
+++ b/Engine/Plugins/Experimental/Text3D/Source/Text3D/Public/Glyph.h
@@ -6,6 +6,8 @@
 #include "MeshDescription.h"
 #include "StaticMeshAttributes.h"
 
+class UStaticMesh;
+
 struct FText3DPolygonGroup
 {
 	int32 FirstVertex;
diff --git a/Engine/Plugins/FX/Niagara/Source/Niagara/Private/NiagaraPerfBaseline.cpp b/Engine/Plugins/FX/Niagara/Source/Niagara/Private/NiagaraPerfBaseline.cpp
--- a/Engine/Plugins/FX/Niagara/Source/Niagara/Private/NiagaraPerfBaseline.cpp
+++ b/Engine/Plugins/FX/Niagara/Source/Niagara/Private/NiagaraPerfBaseline.cpp
@@ -2,6 +2,7 @@
 
 #include "NiagaraPerfBaseline.h"
 #include "Engine/Engine.h"
+#include "Misc/Paths.h"
 #include "NiagaraComponent.h"
 #include "NiagaraEffectType.h"
 #include "NiagaraFunctionLibrary.h"
diff --git a/Engine/Plugins/Importers/USDImporter/Source/USDStage/Private/USDStageActor.cpp b/Engine/Plugins/Importers/USDImporter/Source/USDStage/Private/USDStageActor.cpp
--- a/Engine/Plugins/Importers/USDImporter/Source/USDStage/Private/USDStageActor.cpp
+++ b/Engine/Plugins/Importers/USDImporter/Source/USDStage/Private/USDStageActor.cpp
@@ -4,6 +4,7 @@
 
 #include "Engine/Blueprint.h"
 #include "Engine/BlueprintGeneratedClass.h"
+#include "Engine/Level.h"
 #include "UnrealUSDWrapper.h"
 #include "UObject/Package.h"
 #include "USDAssetCache.h"
diff --git a/Engine/Plugins/Media/ElectraPlayer/Source/ElectraPlayerRuntime/Private/Runtime/Decoder/Linux/AudioDecoderAAC_Linux.cpp b/Engine/Plugins/Media/ElectraPlayer/Source/ElectraPlayerRuntime/Private/Runtime/Decoder/Linux/AudioDecoderAAC_Linux.cpp
--- a/Engine/Plugins/Media/ElectraPlayer/Source/ElectraPlayerRuntime/Private/Runtime/Decoder/Linux/AudioDecoderAAC_Linux.cpp
+++ b/Engine/Plugins/Media/ElectraPlayer/Source/ElectraPlayerRuntime/Private/Runtime/Decoder/Linux/AudioDecoderAAC_Linux.cpp
@@ -5,6 +5,7 @@
 #include "PlayerCore.h"
 #include "PlayerRuntimeGlobal.h"
 
+#include "Stats/Stats.h"
 #include "StreamAccessUnitBuffer.h"
 #include "Decoder/AudioDecoderAAC.h"
 #include "Renderer/RendererBase.h"
diff --git a/Engine/Plugins/Media/ElectraPlayer/Source/ElectraPlayerRuntime/Private/Runtime/Decoder/Linux/VideoDecoderH264_Linux.cpp b/Engine/Plugins/Media/ElectraPlayer/Source/ElectraPlayerRuntime/Private/Runtime/Decoder/Linux/VideoDecoderH264_Linux.cpp
--- a/Engine/Plugins/Media/ElectraPlayer/Source/ElectraPlayerRuntime/Private/Runtime/Decoder/Linux/VideoDecoderH264_Linux.cpp
+++ b/Engine/Plugins/Media/ElectraPlayer/Source/ElectraPlayerRuntime/Private/Runtime/Decoder/Linux/VideoDecoderH264_Linux.cpp
@@ -4,6 +4,7 @@
 #include "HAL/Platform.h"
 
 #include "PlayerCore.h"
+#include "Stats/Stats.h"
 #include "StreamAccessUnitBuffer.h"
 #include "Decoder/VideoDecoderH264.h"
 #include "Renderer/RendererBase.h"
diff --git a/Engine/Plugins/Media/ElectraPlayer/Source/ElectraPlayerRuntime/Private/Runtime/Decoder/Linux/VideoDecoderH265_Linux.cpp b/Engine/Plugins/Media/ElectraPlayer/Source/ElectraPlayerRuntime/Private/Runtime/Decoder/Linux/VideoDecoderH265_Linux.cpp
--- a/Engine/Plugins/Media/ElectraPlayer/Source/ElectraPlayerRuntime/Private/Runtime/Decoder/Linux/VideoDecoderH265_Linux.cpp
+++ b/Engine/Plugins/Media/ElectraPlayer/Source/ElectraPlayerRuntime/Private/Runtime/Decoder/Linux/VideoDecoderH265_Linux.cpp
@@ -4,6 +4,7 @@
 #include "HAL/Platform.h"
 
 #include "PlayerCore.h"
+#include "Stats/Stats.h"
 #include "StreamAccessUnitBuffer.h"
 #include "Decoder/VideoDecoderH265.h"
 #include "Renderer/RendererBase.h"
diff --git a/Engine/Plugins/VirtualProduction/RemoteControl/Source/RemoteControl/Private/RemoteControlBinding.cpp b/Engine/Plugins/VirtualProduction/RemoteControl/Source/RemoteControl/Private/RemoteControlBinding.cpp
--- a/Engine/Plugins/VirtualProduction/RemoteControl/Source/RemoteControl/Private/RemoteControlBinding.cpp
+++ b/Engine/Plugins/VirtualProduction/RemoteControl/Source/RemoteControl/Private/RemoteControlBinding.cpp
@@ -6,6 +6,7 @@
 #include "Engine/Brush.h"
 #include "Engine/Engine.h"
 #include "Engine/EngineTypes.h"
+#include "Engine/Level.h"
 #include "Engine/World.h"
 #include "HAL/IConsoleManager.h"
 #include "IRemoteControlModule.h"
diff --git a/Engine/Plugins/VirtualProduction/Takes/Source/TakesCore/Public/TakesUtils.h b/Engine/Plugins/VirtualProduction/Takes/Source/TakesCore/Public/TakesUtils.h
--- a/Engine/Plugins/VirtualProduction/Takes/Source/TakesCore/Public/TakesUtils.h
+++ b/Engine/Plugins/VirtualProduction/Takes/Source/TakesCore/Public/TakesUtils.h
@@ -2,6 +2,7 @@
 
 #pragma once
 
+#include "AssetRegistry/AssetData.h"
 #include "CoreMinimal.h"
 #include "Misc/PackageName.h"
 #include "MovieScene.h"
diff --git a/Engine/Source/Developer/GameplayDebugger/Private/GameplayDebuggerLocalController.cpp b/Engine/Source/Developer/GameplayDebugger/Private/GameplayDebuggerLocalController.cpp
--- a/Engine/Source/Developer/GameplayDebugger/Private/GameplayDebuggerLocalController.cpp
+++ b/Engine/Source/Developer/GameplayDebugger/Private/GameplayDebuggerLocalController.cpp
@@ -5,6 +5,7 @@
 #include "InputCoreTypes.h"
 #include "Framework/Commands/InputChord.h"
 #include "Components/InputComponent.h"
+#include "Misc/App.h"
 #include "SceneView.h"
 #include "TimerManager.h"
 #include "GameFramework/Pawn.h"
diff --git a/Engine/Source/Runtime/Engine/Private/AudioDeviceManager.cpp b/Engine/Source/Runtime/Engine/Private/AudioDeviceManager.cpp
--- a/Engine/Source/Runtime/Engine/Private/AudioDeviceManager.cpp
+++ b/Engine/Source/Runtime/Engine/Private/AudioDeviceManager.cpp
@@ -8,6 +8,7 @@
 #include "AudioDevice.h"
 #include "AudioMixerDevice.h"
 #include "Engine/World.h"
+#include "Misc/App.h"
 #include "Sound/AudioSettings.h"
 #include "Sound/SoundWave.h"
 #include "GameFramework/GameUserSettings.h"
diff --git a/Engine/Source/Runtime/Engine/Private/LevelInstance/LevelInstanceLevelStreaming.cpp b/Engine/Source/Runtime/Engine/Private/LevelInstance/LevelInstanceLevelStreaming.cpp
--- a/Engine/Source/Runtime/Engine/Private/LevelInstance/LevelInstanceLevelStreaming.cpp
+++ b/Engine/Source/Runtime/Engine/Private/LevelInstance/LevelInstanceLevelStreaming.cpp
@@ -1,6 +1,7 @@
 // Copyright Epic Games, Inc. All Rights Reserved.
 
 #include "LevelInstance/LevelInstanceLevelStreaming.h"
+#include "Engine/Level.h"
 #include "LevelInstance/LevelInstanceInterface.h"
 #include "LevelInstance/LevelInstancePrivate.h"
 #include "LevelInstance/LevelInstanceSubsystem.h"
diff --git a/Engine/Source/Runtime/Engine/Private/ObjectTrace.cpp b/Engine/Source/Runtime/Engine/Private/ObjectTrace.cpp
--- a/Engine/Source/Runtime/Engine/Private/ObjectTrace.cpp
+++ b/Engine/Source/Runtime/Engine/Private/ObjectTrace.cpp
@@ -1,6 +1,7 @@
 // Copyright Epic Games, Inc. All Rights Reserved.
 
 #include "ObjectTrace.h"
+#include "UObject/Package.h"
 
 #include UE_INLINE_GENERATED_CPP_BY_NAME(ObjectTrace)
 
diff --git a/Engine/Source/Runtime/Engine/Private/SkinnedAssetCommon.cpp b/Engine/Source/Runtime/Engine/Private/SkinnedAssetCommon.cpp
--- a/Engine/Source/Runtime/Engine/Private/SkinnedAssetCommon.cpp
+++ b/Engine/Source/Runtime/Engine/Private/SkinnedAssetCommon.cpp
@@ -8,6 +8,7 @@
 #include "UObject/CoreObjectVersion.h"
 #include "UObject/EditorObjectVersion.h"
 #include "Interfaces/ITargetPlatform.h"
+#include "UObject/RenderingObjectVersion.h"
 
 /*-----------------------------------------------------------------------------
 	FSkeletalMeshLODInfo
diff --git a/Engine/Source/Runtime/Engine/Private/WorldPartition/WorldPartition.cpp b/Engine/Source/Runtime/Engine/Private/WorldPartition/WorldPartition.cpp
--- a/Engine/Source/Runtime/Engine/Private/WorldPartition/WorldPartition.cpp
+++ b/Engine/Source/Runtime/Engine/Private/WorldPartition/WorldPartition.cpp
@@ -4,6 +4,7 @@
 	WorldPartition.cpp: UWorldPartition implementation
 =============================================================================*/
 #include "WorldPartition/WorldPartition.h"
+#include "Engine/Level.h"
 #include "UObject/UObjectIterator.h"
 #include "WorldPartition/ActorDescContainer.h"
 #include "WorldPartition/WorldPartitionLog.h"
diff --git a/Engine/Source/Runtime/LevelSequence/Private/LevelSequenceActorSpawner.cpp b/Engine/Source/Runtime/LevelSequence/Private/LevelSequenceActorSpawner.cpp
--- a/Engine/Source/Runtime/LevelSequence/Private/LevelSequenceActorSpawner.cpp
+++ b/Engine/Source/Runtime/LevelSequence/Private/LevelSequenceActorSpawner.cpp
@@ -11,6 +11,7 @@
 #include "Engine/Level.h"
 #include "Engine/LevelStreaming.h"
 #include "LevelUtils.h"
+#include "UObject/Package.h"
 #include "UObject/UObjectGlobals.h"
 
 #include "Evaluation/MovieSceneEvaluationTemplateInstance.h"
