Dataflow : API Updates

diff --git a/Engine/Plugins/Experimental/ChaosFlesh/Source/ChaosFleshEditor/Private/ChaosFlesh/ChaosFleshEditorPlugin.cpp b/Engine/Plugins/Experimental/ChaosFlesh/Source/ChaosFleshEditor/Private/ChaosFlesh/ChaosFleshEditorPlugin.cpp
--- a/Engine/Plugins/Experimental/ChaosFlesh/Source/ChaosFleshEditor/Private/ChaosFlesh/ChaosFleshEditorPlugin.cpp
+++ b/Engine/Plugins/Experimental/ChaosFlesh/Source/ChaosFleshEditor/Private/ChaosFlesh/ChaosFleshEditorPlugin.cpp
@@ -66,13 +66,6 @@ void IChaosFleshEditorPlugin::ShutdownModule()
 	}
 }
 
-TSharedRef<FAssetEditorToolkit> IChaosFleshEditorPlugin::CreateFleshAssetEditor(const EToolkitMode::Type Mode, const TSharedPtr<IToolkitHost>& InitToolkitHost, UObject* FleshAsset)
-{
-	TSharedPtr<FDataflowEditorToolkit> NewAssetEditor = MakeShared<FDataflowEditorToolkit>();
-	NewAssetEditor->InitDataflowEditor(Mode, InitToolkitHost, FleshAsset);
-	return StaticCastSharedPtr<FAssetEditorToolkit>(NewAssetEditor).ToSharedRef();
-}
-
 IMPLEMENT_MODULE(IChaosFleshEditorPlugin, FleshAssetEditor)
 
 
diff --git a/Engine/Plugins/Experimental/ChaosFlesh/Source/ChaosFleshEditor/Public/ChaosFlesh/ChaosFleshEditorPlugin.h b/Engine/Plugins/Experimental/ChaosFlesh/Source/ChaosFleshEditor/Public/ChaosFlesh/ChaosFleshEditorPlugin.h
--- a/Engine/Plugins/Experimental/ChaosFlesh/Source/ChaosFleshEditor/Public/ChaosFlesh/ChaosFleshEditorPlugin.h
+++ b/Engine/Plugins/Experimental/ChaosFlesh/Source/ChaosFleshEditor/Public/ChaosFlesh/ChaosFleshEditorPlugin.h
@@ -22,10 +22,6 @@ class IChaosFleshEditorPlugin : public IModuleInterface
 public:
 	virtual void StartupModule();
 	virtual void ShutdownModule();
-	
-	
-	TSharedRef<FAssetEditorToolkit> CreateFleshAssetEditor(const EToolkitMode::Type Mode, const TSharedPtr< IToolkitHost >& InitToolkitHost, UObject* FleshAsset);
-
 
 	/**
 	 * Singleton-like access to this module's interface.  This is just for convenience!
diff --git a/Engine/Plugins/Experimental/Dataflow/Source/DataflowEditor/Private/Dataflow/DataflowEditorPlugin.cpp b/Engine/Plugins/Experimental/Dataflow/Source/DataflowEditor/Private/Dataflow/DataflowEditorPlugin.cpp
--- a/Engine/Plugins/Experimental/Dataflow/Source/DataflowEditor/Private/Dataflow/DataflowEditorPlugin.cpp
+++ b/Engine/Plugins/Experimental/Dataflow/Source/DataflowEditor/Private/Dataflow/DataflowEditorPlugin.cpp
@@ -49,7 +49,7 @@ void IDataflowEditorPlugin::ShutdownModule()
 TSharedRef<FAssetEditorToolkit> IDataflowEditorPlugin::CreateDataflowAssetEditor(const EToolkitMode::Type Mode, const TSharedPtr<IToolkitHost>& InitToolkitHost, UObject* Asset)
 {
 	TSharedPtr<FDataflowEditorToolkit> NewDataflowAssetEditor = MakeShared<FDataflowEditorToolkit>();
-	NewDataflowAssetEditor->InitDataflowEditor(Mode, InitToolkitHost, Asset);
+	NewDataflowAssetEditor->InitializeEditor(Mode, InitToolkitHost, Asset);
 	return StaticCastSharedPtr<FAssetEditorToolkit>(NewDataflowAssetEditor).ToSharedRef();
 }
 
diff --git a/Engine/Plugins/Experimental/Dataflow/Source/DataflowEditor/Private/Dataflow/DataflowEditorToolkit.cpp b/Engine/Plugins/Experimental/Dataflow/Source/DataflowEditor/Private/Dataflow/DataflowEditorToolkit.cpp
--- a/Engine/Plugins/Experimental/Dataflow/Source/DataflowEditor/Private/Dataflow/DataflowEditorToolkit.cpp
+++ b/Engine/Plugins/Experimental/Dataflow/Source/DataflowEditor/Private/Dataflow/DataflowEditorToolkit.cpp
@@ -87,7 +87,7 @@ bool FDataflowEditorToolkit::CanOpenDataflowEditor(UObject* ObjectToEdit)
 	return Dataflow != nullptr;
 }
 
-void FDataflowEditorToolkit::InitDataflowEditor(const EToolkitMode::Type Mode, const TSharedPtr<IToolkitHost>& InitToolkitHost, UObject* ObjectToEdit)
+void FDataflowEditorToolkit::InitializeEditor(const EToolkitMode::Type Mode, const TSharedPtr<IToolkitHost>& InitToolkitHost, UObject* ObjectToEdit)
 {
 	Asset = nullptr;
 	Dataflow = Cast<UDataflow>(ObjectToEdit);
diff --git a/Engine/Plugins/Experimental/Dataflow/Source/DataflowEditor/Public/Dataflow/DataflowEditorToolkit.h b/Engine/Plugins/Experimental/Dataflow/Source/DataflowEditor/Public/Dataflow/DataflowEditorToolkit.h
--- a/Engine/Plugins/Experimental/Dataflow/Source/DataflowEditor/Public/Dataflow/DataflowEditorToolkit.h
+++ b/Engine/Plugins/Experimental/Dataflow/Source/DataflowEditor/Public/Dataflow/DataflowEditorToolkit.h
@@ -35,7 +35,7 @@ class DATAFLOWEDITOR_API FDataflowEditorToolkit : public FAssetEditorToolkit, pu
 public:
 	static bool CanOpenDataflowEditor(UObject* ObjectToEdit);
 
-	void InitDataflowEditor(const EToolkitMode::Type Mode, const TSharedPtr<IToolkitHost>& InitToolkitHost, UObject* ObjectToEdit);
+	virtual void InitializeEditor(const EToolkitMode::Type Mode, const TSharedPtr<IToolkitHost>& InitToolkitHost, UObject* ObjectToEdit);
 
 	// IToolkit Interface
 	virtual FName GetToolkitFName() const override;
