Add fTetWild parameters to IsoStuffing dataflow node.  Moved TetGrid node to the same file.

diff --git a/Engine/Plugins/Experimental/ChaosFlesh/Source/ChaosFleshNodes/ChaosFleshNodes.Build.cs b/Engine/Plugins/Experimental/ChaosFlesh/Source/ChaosFleshNodes/ChaosFleshNodes.Build.cs
--- a/Engine/Plugins/Experimental/ChaosFlesh/Source/ChaosFleshNodes/ChaosFleshNodes.Build.cs
+++ b/Engine/Plugins/Experimental/ChaosFlesh/Source/ChaosFleshNodes/ChaosFleshNodes.Build.cs
@@ -18,6 +18,7 @@ namespace UnrealBuildTool.Rules
 					"ChaosFlesh",
 					"ChaosFleshEngine",
 					"Engine",
+					"GeometryAlgorithms",
 					"GeometryCore",
 					"MeshConversion",
 					"TetMeshing"
diff --git a/Engine/Plugins/Experimental/ChaosFlesh/Source/ChaosFleshNodes/Private/Dataflow/ChaosFleshEngineAssetNodes.cpp b/Engine/Plugins/Experimental/ChaosFlesh/Source/ChaosFleshNodes/Private/Dataflow/ChaosFleshEngineAssetNodes.cpp
--- a/Engine/Plugins/Experimental/ChaosFlesh/Source/ChaosFleshNodes/Private/Dataflow/ChaosFleshEngineAssetNodes.cpp
+++ b/Engine/Plugins/Experimental/ChaosFlesh/Source/ChaosFleshNodes/Private/Dataflow/ChaosFleshEngineAssetNodes.cpp
@@ -4,8 +4,6 @@
 
 #include "ChaosFlesh/FleshAsset.h"
 #include "ChaosFlesh/FleshCollectionUtility.h"
-#include "Chaos/Utilities.h"
-#include "Chaos/UniformGrid.h"
 #include "Dataflow/DataflowObjectInterface.h"
 #include "Dataflow/DataflowEngine.h"
 #include "GeometryCollection/ManagedArrayCollection.h"
@@ -21,7 +19,6 @@ namespace Dataflow
 		DATAFLOW_NODE_REGISTER_CREATION_FACTORY(FGetFleshAssetDataflowNode);
 		DATAFLOW_NODE_REGISTER_CREATION_FACTORY(FFleshAssetTerminalDataflowNode);
 		DATAFLOW_NODE_REGISTER_CREATION_FACTORY(FImportFleshDataflowNode);
-		DATAFLOW_NODE_REGISTER_CREATION_FACTORY(FConstructTetGridNode);
 		DATAFLOW_NODE_REGISTER_CREATION_FACTORY(FComputeFleshMassNode);
 		DATAFLOW_NODE_REGISTER_CREATION_FACTORY(FComputeFiberFieldNode);
 	}
@@ -77,36 +74,6 @@ void FImportFleshDataflowNode::Evaluate(Dataflow::FContext& Context, const FData
 }
 
 
-void FConstructTetGridNode::Evaluate(Dataflow::FContext& Context, const FDataflowOutput* Out) const
-{
-	if (Out->IsA<FManagedArrayCollection>(&Collection))
-	{
-		Chaos::TVector<int32, 3> Counts(GridCellCount[0], GridCellCount[1], GridCellCount[2]);
-
-		Chaos::TVector<double, 3> MinCorner = -.5 * GridDomain;
-		Chaos::TVector<double, 3> MaxCorner = .5 * GridDomain;
-		Chaos::TUniformGrid<double, 3> Grid(MinCorner, MaxCorner, Counts, 0);
-
-		TArray<FIntVector4> Tets;
-		TArray<FVector> X;
-		Chaos::Utilities::TetMeshFromGrid<double>(Grid, Tets, X);
-
-		UE_LOG(LogChaosFlesh, Display, TEXT("TetGrid generated %d points and %d tetrahedra."), X.Num(), Tets.Num());
-
-		TUniquePtr<FFleshCollection> FleshCollection(
-			FFleshCollection::NewFleshCollection(X, Tets, false, false, false));
-		if (FleshCollection)
-		{
-			SetValue<FManagedArrayCollection>(Context, *FleshCollection.Release(), &Collection);
-		}
-		else
-		{
-			SetValue<FManagedArrayCollection>(Context, Collection, &Collection);
-		}
-	}
-}
-
-
 template <class T> using MType = FManagedArrayCollection::TManagedType<T>;
 
 void FComputeFleshMassNode::Evaluate(Dataflow::FContext& Context, const FDataflowOutput* Out) const
diff --git a/Engine/Plugins/Experimental/ChaosFlesh/Source/ChaosFleshNodes/Private/Dataflow/ChaosFleshTetrahedralNodes.cpp b/Engine/Plugins/Experimental/ChaosFlesh/Source/ChaosFleshNodes/Private/Dataflow/ChaosFleshTetrahedralNodes.cpp
--- a/Engine/Plugins/Experimental/ChaosFlesh/Source/ChaosFleshNodes/Private/Dataflow/ChaosFleshTetrahedralNodes.cpp
+++ b/Engine/Plugins/Experimental/ChaosFlesh/Source/ChaosFleshNodes/Private/Dataflow/ChaosFleshTetrahedralNodes.cpp
@@ -1,13 +1,20 @@
 // Copyright Epic Games, Inc. All Rights Reserved.
 
 #include "Dataflow/ChaosFleshTetrahedralNodes.h"
+#include "Dataflow/ChaosFleshEngineAssetNodes.h"
 
 #include "Chaos/Deformable/Utilities.h"
+#include "Chaos/Utilities.h"
+#include "Chaos/UniformGrid.h"
 #include "ChaosFlesh/FleshCollection.h"
+#include "ChaosFlesh/FleshCollectionUtility.h"
+#include "ChaosLog.h"
 #include "CompGeom/ExactPredicates.h"
+#include "CoreMinimal.h"
 #include "DynamicMesh/DynamicMesh3.h"
 #include "DynamicMesh/DynamicMeshAABBTree3.h"
 #include "Engine/StaticMesh.h"
+#include "FTetWildWrapper.h"
 #include "Generate/IsosurfaceStuffing.h"
 #include "GeometryCollection/ManagedArrayCollection.h"
 #include "MeshDescriptionToDynamicMesh.h"
@@ -18,6 +25,7 @@ namespace Dataflow
 	void ChaosFleshTetrahedralNodes()
 	{
 		DATAFLOW_NODE_REGISTER_CREATION_FACTORY(FGenerateTetrahedralCollectionDataflowNodes);
+		DATAFLOW_NODE_REGISTER_CREATION_FACTORY(FConstructTetGridNode);
 	}
 }
 
@@ -57,6 +65,42 @@ TArray<FIntVector3> GetSurfaceTriangles(const TArray<FIntVector4>& Tets)
 	return FacesSet.Array();
 }
 
+//=============================================================================
+// FConstructTetGridNode
+//=============================================================================
+
+void FConstructTetGridNode::Evaluate(Dataflow::FContext& Context, const FDataflowOutput* Out) const
+{
+	if (Out->IsA<FManagedArrayCollection>(&Collection))
+	{
+		TUniquePtr<FFleshCollection> InCollection(GetValue<DataType>(Context, &Collection).NewCopy<FFleshCollection>());
+
+		Chaos::TVector<int32, 3> Counts(GridCellCount[0], GridCellCount[1], GridCellCount[2]);
+
+		Chaos::TVector<double, 3> MinCorner = -.5 * GridDomain;
+		Chaos::TVector<double, 3> MaxCorner = .5 * GridDomain;
+		Chaos::TUniformGrid<double, 3> Grid(MinCorner, MaxCorner, Counts, 0);
+
+		TArray<FIntVector4> Tets;
+		TArray<FVector> X;
+		Chaos::Utilities::TetMeshFromGrid<double>(Grid, Tets, X);
+
+		UE_LOG(LogChaosFlesh, Display, TEXT("TetGrid generated %d points and %d tetrahedra."), X.Num(), Tets.Num());
+
+		TArray<FIntVector3> Tris = GetSurfaceTriangles(Tets);
+		TUniquePtr<FTetrahedralCollection> TetCollection(
+			FTetrahedralCollection::NewTetrahedralCollection(X, Tris, Tets));
+		InCollection->AppendGeometry(*TetCollection.Get());
+
+		SetValue<DataType>(Context, *(FManagedArrayCollection*)InCollection.Get(), &Collection);
+	}
+}
+
+//=============================================================================
+// FGenerateTetrahedralCollectionDataflowNodes
+//=============================================================================
+
+
 void FGenerateTetrahedralCollectionDataflowNodes::Evaluate(Dataflow::FContext& Context, const FDataflowOutput* Out) const
 {
 	if (Out->IsA<DataType>(&Collection))
@@ -64,57 +108,146 @@ void FGenerateTetrahedralCollectionDataflowNodes::Evaluate(Dataflow::FContext& C
 		TUniquePtr<FFleshCollection> InCollection(GetValue<DataType>(Context, &Collection).NewCopy<FFleshCollection>());
 
 #if WITH_EDITORONLY_DATA
-		if (StaticMesh && NumCells>0 && (-.5<= OffsetPercent && OffsetPercent<=0.5))
+		if (Method == TetMeshingMethod::IsoStuffing)
 		{
-			// make a mesh description for UE::Geometry tools
-			UE::Geometry::FDynamicMesh3 DynamicMesh;
-			FMeshDescriptionToDynamicMesh GetSourceMesh;
-			bool bUsingHiResSource = StaticMesh->IsHiResMeshDescriptionValid();
-			const FMeshDescription* UseSourceMeshDescription =
-				(bUsingHiResSource) ? StaticMesh->GetHiResMeshDescription() : StaticMesh->GetMeshDescription(0);
-			GetSourceMesh.Convert(UseSourceMeshDescription, DynamicMesh);
-
-
-			// Tet mesh generation
-			UE::Geometry::TIsosurfaceStuffing<double> IsosurfaceStuffing;
-			UE::Geometry::FDynamicMeshAABBTree3 Spatial(&DynamicMesh);
-			UE::Geometry::TFastWindingTree<UE::Geometry::FDynamicMesh3> FastWinding(&Spatial);
-			UE::Geometry::FAxisAlignedBox3d Bounds = Spatial.GetBoundingBox();
-			IsosurfaceStuffing.Bounds = FBox(Bounds);
-			double CellSize = Bounds.MaxDim() / NumCells;
-			IsosurfaceStuffing.CellSize = CellSize;
-			IsosurfaceStuffing.IsoValue = .5+OffsetPercent;
-			IsosurfaceStuffing.Implicit = [&FastWinding, &Spatial](FVector3d Pos)
-			{
-				FVector3d Nearest = Spatial.FindNearestPoint(Pos);
-				double WindingSign = FastWinding.FastWindingNumber(Pos) - .5;
-				return FVector3d::Distance(Nearest, Pos) * FMathd::SignNonZero(WindingSign);
-			};
+			EvaluateIsoStuffing(Context, InCollection);
+		}
+		else if (Method == TetMeshingMethod::TetWild)
+		{
+			EvaluateTetWild(Context, InCollection);
+		}
+		else
+		{
+			ensureMsgf(false, TEXT("FGenerateTetrahedralCollectionDataflowNodes unsupported Method."));
+		}
+#else
+		ensureMsgf(false,TEXT("FGenerateTetrahedralCollectionDataflowNodes is an editor only node."));
+#endif
+
+		SetValue<DataType>(Context, *(FManagedArrayCollection*)InCollection.Get(), &Collection);
+	}
+}
+
+void FGenerateTetrahedralCollectionDataflowNodes::EvaluateIsoStuffing(Dataflow::FContext& Context, TUniquePtr<FFleshCollection>& InCollection) const
+{
+#if WITH_EDITORONLY_DATA
+	if (StaticMesh && NumCells > 0 && (-.5 <= OffsetPercent && OffsetPercent <= 0.5))
+	{
+		// make a mesh description for UE::Geometry tools
+		UE::Geometry::FDynamicMesh3 DynamicMesh;
+		FMeshDescriptionToDynamicMesh GetSourceMesh;
+		bool bUsingHiResSource = StaticMesh->IsHiResMeshDescriptionValid();
+		const FMeshDescription* UseSourceMeshDescription =
+			(bUsingHiResSource) ? StaticMesh->GetHiResMeshDescription() : StaticMesh->GetMeshDescription(0);
+		GetSourceMesh.Convert(UseSourceMeshDescription, DynamicMesh);
 
-			IsosurfaceStuffing.Generate();
-			if (IsosurfaceStuffing.Tets.Num() > 0)
+		// Tet mesh generation
+		UE::Geometry::TIsosurfaceStuffing<double> IsosurfaceStuffing;
+		UE::Geometry::FDynamicMeshAABBTree3 Spatial(&DynamicMesh);
+		UE::Geometry::TFastWindingTree<UE::Geometry::FDynamicMesh3> FastWinding(&Spatial);
+		UE::Geometry::FAxisAlignedBox3d Bounds = Spatial.GetBoundingBox();
+		IsosurfaceStuffing.Bounds = FBox(Bounds);
+		double CellSize = Bounds.MaxDim() / NumCells;
+		IsosurfaceStuffing.CellSize = CellSize;
+		IsosurfaceStuffing.IsoValue = .5 + OffsetPercent;
+		IsosurfaceStuffing.Implicit = [&FastWinding, &Spatial](FVector3d Pos)
+		{
+			FVector3d Nearest = Spatial.FindNearestPoint(Pos);
+			double WindingSign = FastWinding.FastWindingNumber(Pos) - .5;
+			return FVector3d::Distance(Nearest, Pos) * FMathd::SignNonZero(WindingSign);
+		};
+
+		UE_LOG(LogChaosFlesh, Display, TEXT("Generating tet mesh via IsoStuffing..."));
+		IsosurfaceStuffing.Generate();
+		if (IsosurfaceStuffing.Tets.Num() > 0)
+		{
+			TArray<FVector> Vertices; Vertices.SetNumUninitialized(IsosurfaceStuffing.Vertices.Num());
+			TArray<FIntVector4> Elements; Elements.SetNumUninitialized(IsosurfaceStuffing.Tets.Num());
+			TArray<FIntVector3> SurfaceElements = GetSurfaceTriangles(IsosurfaceStuffing.Tets);
+
+			for (int32 Tdx = 0; Tdx < IsosurfaceStuffing.Tets.Num(); ++Tdx)
 			{
-				TArray<FVector> Vertices; Vertices.SetNumUninitialized(IsosurfaceStuffing.Vertices.Num());
-				TArray<FIntVector4> Elements; Elements.SetNumUninitialized(IsosurfaceStuffing.Tets.Num());
-				TArray<FIntVector3> SurfaceElements = GetSurfaceTriangles(IsosurfaceStuffing.Tets);
-
-				for (int32 Tdx = 0; Tdx < IsosurfaceStuffing.Tets.Num(); ++Tdx)
-				{
-					Elements[Tdx] = IsosurfaceStuffing.Tets[Tdx];
-				}
-				for (int32 Vdx = 0; Vdx < IsosurfaceStuffing.Vertices.Num(); ++Vdx)
-				{
-					Vertices[Vdx] = IsosurfaceStuffing.Vertices[Vdx];
-				}
-
-				TUniquePtr<FTetrahedralCollection> TetCollection(FTetrahedralCollection::NewTetrahedralCollection(Vertices, SurfaceElements, Elements));
-				InCollection->AppendGeometry(*TetCollection.Get());
+				Elements[Tdx] = IsosurfaceStuffing.Tets[Tdx];
 			}
+			for (int32 Vdx = 0; Vdx < IsosurfaceStuffing.Vertices.Num(); ++Vdx)
+			{
+				Vertices[Vdx] = IsosurfaceStuffing.Vertices[Vdx];
+			}
+
+			TUniquePtr<FTetrahedralCollection> TetCollection(FTetrahedralCollection::NewTetrahedralCollection(Vertices, SurfaceElements, Elements));
+			InCollection->AppendGeometry(*TetCollection.Get());
+
+			UE_LOG(LogChaosFlesh, Display,
+				TEXT("Generated tet mesh via IsoStuffing, num vertices: %d num tets: %d"), Vertices.Num(), Elements.Num());
 		}
+		else
+		{
+			UE_LOG(LogChaosFlesh, Warning, TEXT("IsoStuffing produced 0 tetrahedra."));
+		}
+	}
 #else
-		ensureMsgf(false,TEXT("FGenerateTetrahedralCollectionDataflowNodes is a editor only node."));
+	ensureMsgf(false, TEXT("FGenerateTetrahedralCollectionDataflowNodes is an editor only node."));
 #endif
+}
 
-		SetValue<DataType>(Context, *(FManagedArrayCollection*)InCollection.Get(), &Collection);
+void FGenerateTetrahedralCollectionDataflowNodes::EvaluateTetWild(Dataflow::FContext & Context, TUniquePtr<FFleshCollection>&InCollection) const
+{
+#if WITH_EDITORONLY_DATA
+	if (StaticMesh)
+	{
+		// make a mesh description for UE::Geometry tools
+		UE::Geometry::FDynamicMesh3 DynamicMesh;
+		FMeshDescriptionToDynamicMesh GetSourceMesh;
+		bool bUsingHiResSource = StaticMesh->IsHiResMeshDescriptionValid();
+		const FMeshDescription* UseSourceMeshDescription =
+			(bUsingHiResSource) ? StaticMesh->GetHiResMeshDescription() : StaticMesh->GetMeshDescription(0);
+		GetSourceMesh.Convert(UseSourceMeshDescription, DynamicMesh);
+
+		// Pull out Vertices and Triangles
+		TArray<FVector> Verts;
+		TArray<FIntVector3> Tris;
+		for (FVector V : DynamicMesh.VerticesItr())
+		{
+			Verts.Add(V);
+		}
+		for (UE::Geometry::FIndex3i Tri : DynamicMesh.TrianglesItr())
+		{
+			Tris.Emplace(Tri.A, Tri.B, Tri.C);
+		}
+
+		// Tet mesh generation
+		UE::Geometry::FTetWild::FTetMeshParameters Params;
+		Params.bCoarsen = bCoarsen;
+		Params.bExtractManifoldBoundarySurface = bExtractManifoldBoundarySurface;
+		Params.bSkipSimplification = bSkipSimplification;
+
+		Params.EpsRel = EpsRel;
+		Params.MaxIts = MaxIterations;
+		Params.StopEnergy = StopEnergy;
+		Params.IdealEdgeLength = IdealEdgeLength;
+
+		Params.bInvertOutputTets = bInvertOutputTets;
+
+		TArray<FVector> TetVerts;
+		TArray<FIntVector4> Tets;
+		FProgressCancel Progress;
+		UE_LOG(LogChaosFlesh, Display,TEXT("Generating tet mesh via TetWild..."));
+		if (UE::Geometry::FTetWild::ComputeTetMesh(Params, Verts, Tris, TetVerts, Tets, &Progress))
+		{
+			TArray<FIntVector3> SurfaceElements = GetSurfaceTriangles(Tets);
+			TUniquePtr<FTetrahedralCollection> TetCollection(FTetrahedralCollection::NewTetrahedralCollection(TetVerts, SurfaceElements, Tets));
+			InCollection->AppendGeometry(*TetCollection.Get());
+
+			UE_LOG(LogChaosFlesh, Display,
+				TEXT("Generated tet mesh via TetWild, num vertices: %d num tets: %d"), TetVerts.Num(), Tets.Num());
+		}
+		else
+		{
+			UE_LOG(LogChaosFlesh, Error,
+				TEXT("TetWild tetrahedral mesh generation failed."));
+		}
 	}
+#else
+	ensureMsgf(false, TEXT("FGenerateTetrahedralCollectionDataflowNodes is an editor only node."));
+#endif
 }
diff --git a/Engine/Plugins/Experimental/ChaosFlesh/Source/ChaosFleshNodes/Public/Dataflow/ChaosFleshEngineAssetNodes.h b/Engine/Plugins/Experimental/ChaosFlesh/Source/ChaosFleshNodes/Public/Dataflow/ChaosFleshEngineAssetNodes.h
--- a/Engine/Plugins/Experimental/ChaosFlesh/Source/ChaosFleshNodes/Public/Dataflow/ChaosFleshEngineAssetNodes.h
+++ b/Engine/Plugins/Experimental/ChaosFlesh/Source/ChaosFleshNodes/Public/Dataflow/ChaosFleshEngineAssetNodes.h
@@ -85,35 +85,6 @@ public:
 };
 
 
-USTRUCT(meta = (DataflowFlesh))
-struct FConstructTetGridNode : public FDataflowNode
-{
-	GENERATED_USTRUCT_BODY()
-	DATAFLOW_NODE_DEFINE_INTERNAL(FConstructTetGridNode, "TetGrid", "Flesh", "")
-	DATAFLOW_NODE_RENDER_TYPE(FGeometryCollection::StaticType(), "Collection")
-
-public:
-
-	UPROPERTY(meta = (DataflowOutput, DisplayName = "Collection"))
-	FManagedArrayCollection Collection;
-
-	UPROPERTY(EditAnywhere, Category = "Dataflow", meta = (ClampMin = "1"))
-	FIntVector GridCellCount = FIntVector(10, 10, 10);
-
-	UPROPERTY(EditAnywhere, Category = "Dataflow", meta = (ClampMin = "0.0"))
-	FVector GridDomain = FVector(1.0, 1.0, 1.0);
-
-	FConstructTetGridNode(const Dataflow::FNodeParameters& InParam, FGuid InGuid = FGuid::NewGuid())
-		: FDataflowNode(InParam, InGuid)
-	{
-		RegisterOutputConnection(&Collection);
-	}
-
-	virtual void Evaluate(Dataflow::FContext& Context, const FDataflowOutput* Out) const override;
-};
-
-
-
 USTRUCT(meta = (DataflowFlesh))
 struct FComputeFleshMassNode : public FDataflowNode
 {
diff --git a/Engine/Plugins/Experimental/ChaosFlesh/Source/ChaosFleshNodes/Public/Dataflow/ChaosFleshTetrahedralNodes.h b/Engine/Plugins/Experimental/ChaosFlesh/Source/ChaosFleshNodes/Public/Dataflow/ChaosFleshTetrahedralNodes.h
--- a/Engine/Plugins/Experimental/ChaosFlesh/Source/ChaosFleshNodes/Public/Dataflow/ChaosFleshTetrahedralNodes.h
+++ b/Engine/Plugins/Experimental/ChaosFlesh/Source/ChaosFleshNodes/Public/Dataflow/ChaosFleshTetrahedralNodes.h
@@ -9,7 +9,44 @@
 #include "ChaosFleshTetrahedralNodes.generated.h"
 
 class UStaticMesh;
+class FFleshCollection;
 
+USTRUCT(meta = (DataflowFlesh))
+struct FConstructTetGridNode : public FDataflowNode
+{
+	GENERATED_USTRUCT_BODY()
+	DATAFLOW_NODE_DEFINE_INTERNAL(FConstructTetGridNode, "TetGrid", "Flesh", "")
+	DATAFLOW_NODE_RENDER_TYPE(FGeometryCollection::StaticType(), "Collection")
+
+public:
+	typedef FManagedArrayCollection DataType;
+
+	UPROPERTY(meta = (DataflowInput, DataflowOutput, DisplayName = "Collection"))
+	FManagedArrayCollection Collection;
+
+	UPROPERTY(EditAnywhere, Category = "Dataflow", meta = (ClampMin = "1"))
+	FIntVector GridCellCount = FIntVector(10, 10, 10);
+
+	UPROPERTY(EditAnywhere, Category = "Dataflow", meta = (ClampMin = "0.0"))
+	FVector GridDomain = FVector(10.0, 10.0, 10.0);
+
+	FConstructTetGridNode(const Dataflow::FNodeParameters& InParam, FGuid InGuid = FGuid::NewGuid())
+		: FDataflowNode(InParam, InGuid)
+	{
+		RegisterInputConnection(&Collection);
+		RegisterOutputConnection(&Collection);
+	}
+
+	virtual void Evaluate(Dataflow::FContext& Context, const FDataflowOutput* Out) const override;
+};
+
+
+UENUM()
+enum TetMeshingMethod
+{
+	IsoStuffing		UMETA(DisplayName = "IsoStuffing"),
+	TetWild			UMETA(DisplayName = "TetWild"),
+};
 
 USTRUCT(meta = (DataflowFlesh))
 struct FGenerateTetrahedralCollectionDataflowNodes : public FDataflowNode
@@ -22,11 +59,58 @@ public:
 	typedef FManagedArrayCollection DataType;
 
 	UPROPERTY(EditAnywhere, Category = "Dataflow")
+	TEnumAsByte<TetMeshingMethod> Method = TetMeshingMethod::IsoStuffing;
+
+	//
+	// IsoStuffing
+	//
+
+	UPROPERTY(EditAnywhere, Category = "Dataflow", meta = (ClampMin = "1", EditCondition = "Method == TetMeshingMethod::IsoStuffing", EditConditionHides))
 	int32 NumCells = 32;
 
-	UPROPERTY(EditAnywhere, Category = "Dataflow")
+	UPROPERTY(EditAnywhere, Category = "Dataflow", meta = (ClampMin = "-0.5", ClampMax = "0.5", EditCondition = "Method == TetMeshingMethod::IsoStuffing", EditConditionHides))
 	double OffsetPercent = 0.05;
 
+	//
+	// TetWild
+	//
+
+	//! Energy at which to stop optimizing tet quality and accept the result.
+	UPROPERTY(EditAnywhere, Category = "Dataflow", meta = (ClampMin = "0.0", EditCondition = "Method == TetMeshingMethod::TetWild", EditConditionHides))
+	double IdealEdgeLength = 1.0;
+
+	//! Maximum number of optimization iterations.
+	UPROPERTY(EditAnywhere, Category = "Dataflow", meta = (ClampMin = "1", EditCondition = "Method == TetMeshingMethod::TetWild", EditConditionHides))
+	int32 MaxIterations = 80;
+
+	//! Energy at which to stop optimizing tet quality and accept the result.
+	UPROPERTY(EditAnywhere, Category = "Dataflow", meta = (ClampMin = "0.0", EditCondition = "Method == TetMeshingMethod::TetWild", EditConditionHides))
+	double StopEnergy = 10;
+
+	//! Relative tolerance, controlling how closely the mesh must follow the input surface.
+	UPROPERTY(EditAnywhere, Category = "Dataflow", meta = (ClampMin = "0.0", EditCondition = "Method == TetMeshingMethod::TetWild", EditConditionHides))
+	double EpsRel = 1e-3;
+
+	//! Coarsen the tet mesh result.
+	UPROPERTY(EditAnywhere, Category = "Dataflow", meta = (EditCondition = "Method == TetMeshingMethod::TetWild", EditConditionHides))
+	bool bCoarsen = false;
+
+	//! Enforce that the output boundary surface should be manifold.
+	UPROPERTY(EditAnywhere, Category = "Dataflow", meta = (EditCondition = "Method == TetMeshingMethod::TetWild", EditConditionHides))
+	bool bExtractManifoldBoundarySurface = false;
+
+	//! Skip the initial simplification step.
+	UPROPERTY(EditAnywhere, Category = "Dataflow", meta = (EditCondition = "Method == TetMeshingMethod::TetWild", EditConditionHides))
+	bool bSkipSimplification = false;
+
+	//! Invert tetrahedra.
+	UPROPERTY(EditAnywhere, Category = "Dataflow", meta = (EditCondition = "Method == TetMeshingMethod::TetWild", EditConditionHides))
+	bool bInvertOutputTets = false;
+
+	//
+	// Common
+	//
+
 	UPROPERTY(EditAnywhere, Category = "Dataflow")
 	TObjectPtr<const UStaticMesh> StaticMesh = nullptr;
 
@@ -43,6 +127,9 @@ public:
 
 	virtual void Evaluate(Dataflow::FContext& Context, const FDataflowOutput* Out) const override;
 
+protected:
+	void EvaluateIsoStuffing(Dataflow::FContext& Context, TUniquePtr<FFleshCollection>& InCollection) const;
+	void EvaluateTetWild(Dataflow::FContext& Context, TUniquePtr<FFleshCollection>& InCollection) const;
 };
 
 namespace Dataflow
