Get StaticMesh from the context.

diff --git a/Engine/Plugins/Experimental/ChaosFlesh/Source/ChaosFleshNodes/Private/Dataflow/ChaosFleshTetrahedralNodes.cpp b/Engine/Plugins/Experimental/ChaosFlesh/Source/ChaosFleshNodes/Private/Dataflow/ChaosFleshTetrahedralNodes.cpp
--- a/Engine/Plugins/Experimental/ChaosFlesh/Source/ChaosFleshNodes/Private/Dataflow/ChaosFleshTetrahedralNodes.cpp
+++ b/Engine/Plugins/Experimental/ChaosFlesh/Source/ChaosFleshNodes/Private/Dataflow/ChaosFleshTetrahedralNodes.cpp
@@ -106,15 +106,16 @@ void FGenerateTetrahedralCollectionDataflowNodes::Evaluate(Dataflow::FContext& C
 	if (Out->IsA<DataType>(&Collection))
 	{
 		TUniquePtr<FFleshCollection> InCollection(GetValue<DataType>(Context, &Collection).NewCopy<FFleshCollection>());
+		TObjectPtr<const UStaticMesh> InStaticMesh(GetValue<TObjectPtr<const UStaticMesh>>(Context, &StaticMesh));
 
 #if WITH_EDITORONLY_DATA
 		if (Method == TetMeshingMethod::IsoStuffing)
 		{
-			EvaluateIsoStuffing(Context, InCollection);
+			EvaluateIsoStuffing(Context, InCollection, InStaticMesh);
 		}
 		else if (Method == TetMeshingMethod::TetWild)
 		{
-			EvaluateTetWild(Context, InCollection);
+			EvaluateTetWild(Context, InCollection, InStaticMesh);
 		}
 		else
 		{
@@ -128,17 +129,20 @@ void FGenerateTetrahedralCollectionDataflowNodes::Evaluate(Dataflow::FContext& C
 	}
 }
 
-void FGenerateTetrahedralCollectionDataflowNodes::EvaluateIsoStuffing(Dataflow::FContext& Context, TUniquePtr<FFleshCollection>& InCollection) const
+void FGenerateTetrahedralCollectionDataflowNodes::EvaluateIsoStuffing(
+	Dataflow::FContext& Context, 
+	TUniquePtr<FFleshCollection>& InCollection,
+	TObjectPtr<const UStaticMesh>& InStaticMesh) const
 {
 #if WITH_EDITORONLY_DATA
-	if (StaticMesh && NumCells > 0 && (-.5 <= OffsetPercent && OffsetPercent <= 0.5))
+	if (InStaticMesh && NumCells > 0 && (-.5 <= OffsetPercent && OffsetPercent <= 0.5))
 	{
 		// make a mesh description for UE::Geometry tools
 		UE::Geometry::FDynamicMesh3 DynamicMesh;
 		FMeshDescriptionToDynamicMesh GetSourceMesh;
-		bool bUsingHiResSource = StaticMesh->IsHiResMeshDescriptionValid();
+		bool bUsingHiResSource = InStaticMesh->IsHiResMeshDescriptionValid();
 		const FMeshDescription* UseSourceMeshDescription =
-			(bUsingHiResSource) ? StaticMesh->GetHiResMeshDescription() : StaticMesh->GetMeshDescription(0);
+			(bUsingHiResSource) ? InStaticMesh->GetHiResMeshDescription() : InStaticMesh->GetMeshDescription(0);
 		GetSourceMesh.Convert(UseSourceMeshDescription, DynamicMesh);
 
 		// Tet mesh generation
@@ -190,17 +194,20 @@ void FGenerateTetrahedralCollectionDataflowNodes::EvaluateIsoStuffing(Dataflow::
 #endif
 }
 
-void FGenerateTetrahedralCollectionDataflowNodes::EvaluateTetWild(Dataflow::FContext & Context, TUniquePtr<FFleshCollection>&InCollection) const
+void FGenerateTetrahedralCollectionDataflowNodes::EvaluateTetWild(
+	Dataflow::FContext& Context, 
+	TUniquePtr<FFleshCollection>& InCollection,
+	TObjectPtr<const UStaticMesh>& InStaticMesh) const
 {
 #if WITH_EDITORONLY_DATA
-	if (StaticMesh)
+	if (StaticMesh && InStaticMesh)
 	{
 		// make a mesh description for UE::Geometry tools
 		UE::Geometry::FDynamicMesh3 DynamicMesh;
 		FMeshDescriptionToDynamicMesh GetSourceMesh;
-		bool bUsingHiResSource = StaticMesh->IsHiResMeshDescriptionValid();
+		bool bUsingHiResSource = InStaticMesh->IsHiResMeshDescriptionValid();
 		const FMeshDescription* UseSourceMeshDescription =
-			(bUsingHiResSource) ? StaticMesh->GetHiResMeshDescription() : StaticMesh->GetMeshDescription(0);
+			(bUsingHiResSource) ? InStaticMesh->GetHiResMeshDescription() : InStaticMesh->GetMeshDescription(0);
 		GetSourceMesh.Convert(UseSourceMeshDescription, DynamicMesh);
 
 		// Pull out Vertices and Triangles
diff --git a/Engine/Plugins/Experimental/ChaosFlesh/Source/ChaosFleshNodes/Public/Dataflow/ChaosFleshTetrahedralNodes.h b/Engine/Plugins/Experimental/ChaosFlesh/Source/ChaosFleshNodes/Public/Dataflow/ChaosFleshTetrahedralNodes.h
--- a/Engine/Plugins/Experimental/ChaosFlesh/Source/ChaosFleshNodes/Public/Dataflow/ChaosFleshTetrahedralNodes.h
+++ b/Engine/Plugins/Experimental/ChaosFlesh/Source/ChaosFleshNodes/Public/Dataflow/ChaosFleshTetrahedralNodes.h
@@ -128,8 +128,8 @@ public:
 	virtual void Evaluate(Dataflow::FContext& Context, const FDataflowOutput* Out) const override;
 
 protected:
-	void EvaluateIsoStuffing(Dataflow::FContext& Context, TUniquePtr<FFleshCollection>& InCollection) const;
-	void EvaluateTetWild(Dataflow::FContext& Context, TUniquePtr<FFleshCollection>& InCollection) const;
+	void EvaluateIsoStuffing(Dataflow::FContext& Context, TUniquePtr<FFleshCollection>& InCollection, TObjectPtr<const UStaticMesh>& InStaticMesh) const;
+	void EvaluateTetWild(Dataflow::FContext& Context, TUniquePtr<FFleshCollection>& InCollection, TObjectPtr<const UStaticMesh>& InStaticMesh) const;
 };
 
 namespace Dataflow
