Add tetrahedral bindings dataflow node, and tetrahedral bindings facade.  Fix issues with TetrahedralCollection tet start and count attributes not being handled properly during copy.  Add global indexing option to TriangleMesh point normals, and expose BVH and SmoothProject functions via CHAOS_API.  Add accessors to ManagedArrayAccessor.  Fix indexing error in SmoothProject, and undo subsequent contortions.

diff --git a/Engine/Plugins/Experimental/ChaosFlesh/Source/ChaosFlesh/Private/ChaosFlesh/TetrahedralCollection.cpp b/Engine/Plugins/Experimental/ChaosFlesh/Source/ChaosFlesh/Private/ChaosFlesh/TetrahedralCollection.cpp
--- a/Engine/Plugins/Experimental/ChaosFlesh/Source/ChaosFlesh/Private/ChaosFlesh/TetrahedralCollection.cpp
+++ b/Engine/Plugins/Experimental/ChaosFlesh/Source/ChaosFlesh/Private/ChaosFlesh/TetrahedralCollection.cpp
@@ -8,9 +8,12 @@ DEFINE_LOG_CATEGORY_STATIC(FTetrahedralCollectionLogging, Log, All);
 
 // Groups
 const FName FTetrahedralCollection::TetrahedralGroup = "Tetrahedral";
+const FName FTetrahedralCollection::BindingsGroup = "Bindings";
 
 // Attributes
 const FName FTetrahedralCollection::TetrahedronAttribute("Tetrahedron");
+const FName FTetrahedralCollection::TetrahedronStartAttribute("TetrahedronStart");
+const FName FTetrahedralCollection::TetrahedronCountAttribute("TetrahedronCount");
 const FName FTetrahedralCollection::IncidentElementsAttribute("IncidentElements");
 const FName FTetrahedralCollection::IncidentElementsLocalIndexAttribute("IncidentElementsLocalIndex");
 
@@ -34,8 +37,8 @@ void FTetrahedralCollection::Construct()
 	AddExternalAttribute<TArray<int32>>(FTetrahedralCollection::IncidentElementsLocalIndexAttribute, FGeometryCollection::VerticesGroup, IncidentElementsLocalIndex);
 
 	// Geometry Group
-	AddExternalAttribute<int32>("TetrahedronStart", FGeometryCollection::GeometryGroup, TetrahedronStart);
-	AddExternalAttribute<int32>("TetrahedronCount", FGeometryCollection::GeometryGroup, TetrahedronCount);
+	AddExternalAttribute<int32>(FTetrahedralCollection::TetrahedronStartAttribute, FGeometryCollection::GeometryGroup, TetrahedronStart, TetrahedronDependency);
+	AddExternalAttribute<int32>(FTetrahedralCollection::TetrahedronCountAttribute, FGeometryCollection::GeometryGroup, TetrahedronCount, TetrahedronDependency);
 }
 
 
@@ -52,13 +55,22 @@ void FTetrahedralCollection::SetDefaults(FName Group, uint32 StartSize, uint32 N
 	}
 }
 
-FTetrahedralCollection* FTetrahedralCollection::NewTetrahedralCollection(const TArray<FVector>& Vertices, const TArray<FIntVector3>& SurfaceElements, const TArray<FIntVector4>& Elements, bool bReverseVertexOrder)
+FTetrahedralCollection* FTetrahedralCollection::NewTetrahedralCollection(
+	const TArray<FVector>& Vertices, 
+	const TArray<FIntVector3>& SurfaceElements, 
+	const TArray<FIntVector4>& Elements, 
+	bool bReverseVertexOrder)
 {
 	FTetrahedralCollection* Collection = new FTetrahedralCollection();
 	FTetrahedralCollection::Init(Collection, Vertices, SurfaceElements, Elements, bReverseVertexOrder);
 	return Collection;
 }
-void FTetrahedralCollection::Init(FTetrahedralCollection* Collection, const TArray<FVector>& Vertices, const TArray<FIntVector3>& SurfaceElements, const TArray<FIntVector4>& Elements, bool bReverseVertexOrder)
+void FTetrahedralCollection::Init(
+	FTetrahedralCollection* Collection, 
+	const TArray<FVector>& Vertices, 
+	const TArray<FIntVector3>& SurfaceElements, 
+	const TArray<FIntVector4>& Elements, 
+	bool bReverseVertexOrder)
 {
 	if (Collection)
 	{
@@ -81,6 +93,8 @@ void FTetrahedralCollection::Init(FTetrahedralCollection* Collection, const TArr
 
 		Super::Init(Collection, RawVertexArray, RawIndicesArray, bReverseVertexOrder);
 
+		Collection->TetrahedronStart[0] = Collection->Tetrahedron.Num();
+		Collection->TetrahedronCount[0] = Elements.Num();
 		Collection->AddElements(Elements.Num(), FTetrahedralCollection::TetrahedralGroup);
 		for (int i = 0; i < Elements.Num(); i++)
 		{
@@ -99,6 +113,7 @@ int32 FTetrahedralCollection::AppendGeometry(
 	const FTransform& TransformRoot)
 {
 	const int32 VerticesIndex = NumElements(FGeometryCollection::VerticesGroup);
+	const int32 NumGeometry = NumElements(FGeometryCollection::GeometryGroup);
 
 	const int32 ID = 
 		Super::AppendGeometry(
@@ -117,6 +132,15 @@ int32 FTetrahedralCollection::AppendGeometry(
 		Tetrahedron[TetsIndex + Idx] = FIntVector4(VerticesIndex, VerticesIndex, VerticesIndex, VerticesIndex) + Other.Tetrahedron[Idx];
 	}
 
+	// --- GEOMETRY GROUP --
+
+	check(TetrahedronStart.Num() == NumGeometry + Other.TetrahedronStart.Num());
+	for (int32 Idx = 0; Idx < Other.TetrahedronStart.Num(); Idx++)
+	{
+		TetrahedronStart[NumGeometry + Idx] = NumTets + Other.TetrahedronStart[Idx];
+		TetrahedronCount[NumGeometry + Idx] = Other.TetrahedronCount[Idx];
+	}
+
 	return ID;
 }
 
@@ -196,4 +220,4 @@ void FTetrahedralCollection::ReorderTetrahedralElements(const TArray<int32>& New
 	}
 
 	Super::ReorderElements(TetrahedralGroup, NewTetOrder);
-}
\ No newline at end of file
+}
diff --git a/Engine/Plugins/Experimental/ChaosFlesh/Source/ChaosFlesh/Public/ChaosFlesh/TetrahedralCollection.h b/Engine/Plugins/Experimental/ChaosFlesh/Source/ChaosFlesh/Public/ChaosFlesh/TetrahedralCollection.h
--- a/Engine/Plugins/Experimental/ChaosFlesh/Source/ChaosFlesh/Public/ChaosFlesh/TetrahedralCollection.h
+++ b/Engine/Plugins/Experimental/ChaosFlesh/Source/ChaosFlesh/Public/ChaosFlesh/TetrahedralCollection.h
@@ -54,9 +54,10 @@ public:
 	void ReorderTetrahedralElements(const TArray<int32>& NewOrder);
 
 	/*
-	*  Attribute Groupes
+	*  Attribute Groups
 	*/
 	static const FName TetrahedralGroup;
+	static const FName BindingsGroup;
 
 	/*
 	*  Tetrahedron Attribute
@@ -65,7 +66,9 @@ public:
 	static const FName TetrahedronAttribute;
 	TManagedArray<FIntVector4> Tetrahedron;
 
+	static const FName TetrahedronStartAttribute;
 	TManagedArray<int32> TetrahedronStart;
+	static const FName TetrahedronCountAttribute;
 	TManagedArray<int32> TetrahedronCount;
 
 	/**
diff --git a/Engine/Plugins/Experimental/ChaosFlesh/Source/ChaosFleshNodes/Private/Dataflow/ChaosFleshBindingsNodes.cpp b/Engine/Plugins/Experimental/ChaosFlesh/Source/ChaosFleshNodes/Private/Dataflow/ChaosFleshBindingsNodes.cpp
new file mode 100644
--- /dev/null
+++ b/Engine/Plugins/Experimental/ChaosFlesh/Source/ChaosFleshNodes/Private/Dataflow/ChaosFleshBindingsNodes.cpp
@@ -0,0 +1,304 @@
+// Copyright Epic Games, Inc. All Rights Reserved.
+
+#include "Dataflow/ChaosFleshBindingsNodes.h"
+
+#include "Chaos/BoundingVolumeHierarchy.h"
+#include "Chaos/Tetrahedron.h"
+#include "Chaos/TriangleMesh.h"
+#include "ChaosFlesh/FleshCollection.h"
+#include "Engine/StaticMesh.h"
+#include "Engine/SkeletalMesh.h"
+#include "GeometryCollection/Facades/CollectionTetrahedralBindingsFacade.h"
+#include "Rendering/SkeletalMeshRenderData.h"
+#include "StaticMeshResources.h"
+
+namespace Dataflow
+{
+	void ChaosFleshBindingsNodes()
+	{
+		DATAFLOW_NODE_REGISTER_CREATION_FACTORY(FGenerateBindings);
+	}
+}
+
+void
+FGenerateBindings::Evaluate(Dataflow::FContext& Context, const FDataflowOutput* Out) const
+{
+	if (Out->IsA<DataType>(&Collection))
+	{
+		DataType InCollection = GetValue<DataType>(Context, &Collection); // Deep copy
+
+		TManagedArray<FIntVector4>* Tetrahedron = 
+			InCollection.FindAttribute<FIntVector4>(
+				FTetrahedralCollection::TetrahedronAttribute, FTetrahedralCollection::TetrahedralGroup);
+		TManagedArray<int32>* TetrahedronStart =
+			InCollection.FindAttribute<int32>(
+				"TetrahedronStart", FGeometryCollection::GeometryGroup);
+		TManagedArray<int32>* TetrahedronCount =
+			InCollection.FindAttribute<int32>(
+				"TetrahedronCount", FGeometryCollection::GeometryGroup);
+
+		TManagedArray<FIntVector>* Triangle =
+			InCollection.FindAttribute<FIntVector>(
+				"Indices", FGeometryCollection::FacesGroup);
+		TManagedArray<int32>* FacesStart =
+			InCollection.FindAttribute<int32>(
+				"FaceStart", FGeometryCollection::GeometryGroup);
+		TManagedArray<int32>* FacesCount =
+			InCollection.FindAttribute<int32>(
+				"FaceCount", FGeometryCollection::GeometryGroup);
+
+		TManagedArray<FVector3f>* Vertex = 
+			InCollection.FindAttribute<FVector3f>(
+				"Vertex", "Vertices");
+
+		TObjectPtr<const USkeletalMesh> SkeletalMesh = GetValue<TObjectPtr<const USkeletalMesh>>(Context, &SkeletalMeshIn);
+		TObjectPtr<const UStaticMesh> StaticMesh = GetValue<TObjectPtr<const UStaticMesh>>(Context, &StaticMeshIn);
+		if ((SkeletalMesh || StaticMesh) && Tetrahedron && Triangle && Vertex)
+		{
+			// Extract positions to bind
+			FString MeshId;
+			TArray<TArray<FVector3f>> MeshVertices;
+			if (SkeletalMesh)
+			{
+				FPrimaryAssetId Id = SkeletalMesh->GetPrimaryAssetId();
+				if (Id.IsValid())
+				{
+					MeshId = Id.ToString();
+				}
+				else
+				{
+					MeshId = SkeletalMesh->GetName();
+				}
+
+				FSkeletalMeshRenderData* RenderData = SkeletalMesh->GetResourceForRendering();
+
+				MeshVertices.SetNum(RenderData->LODRenderData.Num());
+				for (int32 i = 0; i < RenderData->LODRenderData.Num(); i++)
+				{
+					FSkeletalMeshLODRenderData* LODRenderData = &RenderData->LODRenderData[i];
+					const FPositionVertexBuffer& PositionVertexBuffer =
+						LODRenderData->StaticVertexBuffers.PositionVertexBuffer;
+
+					TArray<FVector3f>& Vertices = MeshVertices[i];
+					Vertices.SetNumUninitialized(PositionVertexBuffer.GetNumVertices());
+					for (uint32 j = 0; j < PositionVertexBuffer.GetNumVertices(); j++)
+					{
+						const FVector3f& Pos = PositionVertexBuffer.VertexPosition(j);
+						Vertices[j] = Pos;
+					}
+				}
+			}
+			else // StaticMesh
+			{
+				//StaticMesh->GetMeshId(MeshId); // not available at runtime!
+				//MeshId = RenderData->DerivedDataKey; // same problem
+				FPrimaryAssetId Id = StaticMesh->GetPrimaryAssetId();
+				if (Id.IsValid())
+				{
+					MeshId = Id.ToString();
+				}
+				else
+				{
+					MeshId = StaticMesh->GetName();
+				}
+
+				const FStaticMeshRenderData* RenderData = StaticMesh->GetRenderData();
+				const int32 NumLOD = RenderData->LODResources.Num();
+				MeshVertices.SetNum(NumLOD);
+				for (int32 i = 0; i < NumLOD; i++)
+				{
+					const FStaticMeshLODResources& LODResources = RenderData->LODResources[i];
+					const FPositionVertexBuffer& PositionVertexBuffer = LODResources.VertexBuffers.PositionVertexBuffer;
+					
+					TArray<FVector3f>& Vertices = MeshVertices[i];
+					Vertices.SetNumUninitialized(PositionVertexBuffer.GetNumVertices());
+					for (uint32 j = 0; j < PositionVertexBuffer.GetNumVertices(); j++)
+					{
+						const FVector3f& Pos = PositionVertexBuffer.VertexPosition(j);
+						Vertices[j] = Pos;
+					}
+				}
+
+/*				// Geom lib does conversion like this:
+				const FMeshDescription* MeshDescription =
+					StaticMesh->IsHiResMeshDescriptionValid() ?
+					StaticMesh->GetHiResMeshDescription() :
+					StaticMesh->GetMeshDescription(0);
+				TVertexAttributesConstRef<const FVector3f> VertexPositions = MeshDescription->GetVertexPositions();
+				//TVertexAttributesConstRef<FVector3f> VertexPositions = MeshDescription->GetVertexPositions();
+
+				MeshVertices.SetNum(1);
+				MeshVertices[0].Reserve(MeshDescription->Vertices().Num());
+				for (const FVertexID VertexID : MeshDescription->Vertices().GetElementIDs())
+				{
+					const FVector3f Pos = VertexPositions.Get(VertexID);
+					MeshVertices[0].Add(Pos);
+				}
+*/
+			}
+
+			for (int32 TetMeshIdx = 0; TetMeshIdx < TetrahedronStart->Num(); TetMeshIdx++)
+			{
+				const int32 TetMeshStart = (*TetrahedronStart)[TetMeshIdx];
+				const int32 TetMeshCount = (*TetrahedronCount)[TetMeshIdx];
+
+				// Build Tetrahedra
+				TArray<Chaos::TTetrahedron<Chaos::FReal>> Tets;
+				TArray<Chaos::TTetrahedron<Chaos::FReal>*> BVHTetPtrs;
+				Tets.SetNumUninitialized(TetMeshCount);
+				BVHTetPtrs.SetNumUninitialized(TetMeshCount);
+				for (int32 i = TetMeshStart; i < TetMeshStart + TetMeshCount; i++)
+				{
+					const FIntVector4& Tet = (*Tetrahedron)[i];
+					Tets[i] = Chaos::TTetrahedron<Chaos::FReal>(
+						(*Vertex)[Tet[0]],
+						(*Vertex)[Tet[1]],
+						(*Vertex)[Tet[2]],
+						(*Vertex)[Tet[3]]);
+					BVHTetPtrs[i] = &Tets[i];
+				}
+
+				// Init BVH for tetrahedra.
+				Chaos::TBoundingVolumeHierarchy<
+					TArray<Chaos::TTetrahedron<Chaos::FReal>*>, 
+					TArray<int32>, 
+					Chaos::FReal, 
+					3> TetBVH(BVHTetPtrs);
+
+				//
+				// Init boundary mesh for projections.
+				//
+
+				const int32 TriMeshStart = (*FacesStart)[TetMeshIdx];
+				const int32 TriMeshCount = (*FacesCount)[TetMeshIdx];
+				Chaos::FTriangleMesh TetBoundaryMesh;
+				if (TriMeshCount == Triangle->GetConstArray().Num())
+				{
+					TetBoundaryMesh.Init(
+						reinterpret_cast<const TArray<Chaos::TVec3<int32>>&>(Triangle->GetConstArray()), 0, -1, false);
+				}
+				else
+				{
+					TArray<Chaos::TVec3<int32>> Faces; Faces.SetNumUninitialized(TriMeshCount);
+					for (int32 i = TriMeshStart; i < TriMeshStart + TriMeshCount; i++)
+					{
+						Faces[i] = reinterpret_cast<const Chaos::TVec3<int32>&>(Triangle->GetConstArray()[i]);
+					}
+					TetBoundaryMesh.Init(Faces, 0, -1, false);
+				}
+				
+				// Promote vertices to double because that's what FTriangleMesh wants.
+				TArray<Chaos::FVec3> VertexD; VertexD.SetNumUninitialized(Vertex->Num());
+				for (int32 i = 0; i < VertexD.Num(); i++)
+				{
+					VertexD[i] = Chaos::FVec3((*Vertex)[i][0], (*Vertex)[i][1], (*Vertex)[i][2]);
+				}
+				TConstArrayView<Chaos::TVec3<Chaos::FRealDouble>> VertexDView(VertexD);
+				TArray<Chaos::FVec3> PointNormals = TetBoundaryMesh.GetPointNormals(VertexDView, false, true);
+
+				Chaos::FTriangleMesh::TBVHType<Chaos::FRealDouble> TetBoundaryBVH;
+				//TetBoundaryMesh.BuildBVH_(VertexDView, TetBoundaryBVH);
+				TetBoundaryMesh.BuildBVH(VertexDView, TetBoundaryBVH);
+
+				//
+				// Do intersection tests against tets, then the surface.
+				//
+
+				TArray<TArray<FIntVector4>> Parents; Parents.SetNum(MeshVertices.Num());
+				TArray<TArray<FVector4f>> Weights;	 Weights.SetNum(MeshVertices.Num());
+				TArray<TArray<FVector3f>> Offsets;	 Offsets.SetNum(MeshVertices.Num());
+				TArray<int32> Orphans;
+				for (int32 LOD = 0; LOD < MeshVertices.Num(); LOD++)
+				{
+					Parents[LOD].SetNumUninitialized(MeshVertices[LOD].Num());
+					Weights[LOD].SetNumUninitialized(MeshVertices[LOD].Num());
+					Offsets[LOD].SetNumUninitialized(MeshVertices[LOD].Num());
+
+					TArray<int32> TetIntersections; TetIntersections.Reserve(64);
+					for (int32 i = 0; i < MeshVertices[LOD].Num(); i++)
+					{
+						Parents[LOD][i] = FIntVector4(INDEX_NONE);
+						Weights[LOD][i] = FVector4f(0);
+						Offsets[LOD][i] = FVector3f(0);
+
+						const FVector3f& Pos = MeshVertices[LOD][i];
+						TetIntersections = TetBVH.FindAllIntersections(Chaos::TVec3<Chaos::FReal>(Pos[0], Pos[1], Pos[2]));
+						int32 j = 0;
+						for (; j < TetIntersections.Num(); j++)
+						{
+							const int32 TetIdx = TetIntersections[j];
+							if (Tets[TetIdx].RobustInside(Pos, -1.0e-4))
+							{
+								Parents[LOD][i] = (*Tetrahedron)[TetIdx];
+								Chaos::TVector<Chaos::FReal, 4> WeightsD = Tets[TetIdx].GetBarycentricCoordinates(Pos);
+								Weights[LOD][i] = FVector4f(WeightsD[0], WeightsD[1], WeightsD[2], WeightsD[3]);
+								Offsets[LOD][i] = FVector3f(0);
+								break;
+							}
+						}
+						if (j == TetIntersections.Num())
+						{
+							// This vertex didn't land inside any tetrahedra. Project it to the tet boundary surface.
+							int32 TriIdx = INDEX_NONE;
+							Chaos::FVec3 TriWeights;
+							//if (TetBoundaryMesh.SmoothProject_(
+							if (TetBoundaryMesh.SmoothProject(
+								TetBoundaryBVH,
+								VertexDView,
+								PointNormals, Pos,
+								TriIdx, TriWeights, SurfaceProjectionIterations))
+							{
+								const FIntVector& Tri = (*Triangle)[TriIdx];
+								Parents[LOD][i][0] = Tri[0];
+								Parents[LOD][i][1] = Tri[1];
+								Parents[LOD][i][2] = Tri[2];
+								Parents[LOD][i][3] = INDEX_NONE;
+
+								Weights[LOD][i][0] = TriWeights[0];
+								Weights[LOD][i][1] = TriWeights[1];
+								Weights[LOD][i][2] = TriWeights[2];
+								Weights[LOD][i][3] = 0.0;
+
+								const FVector3f EmbeddedPos =
+									TriWeights[0] * Vertex->GetConstArray()[Tri[0]] +
+									TriWeights[1] * Vertex->GetConstArray()[Tri[1]] +
+									TriWeights[2] * Vertex->GetConstArray()[Tri[2]];
+								Offsets[LOD][i] = Pos - EmbeddedPos;
+							}
+							else
+							{
+								// Despair...
+								Orphans.Add(i);
+
+								Parents[LOD][i][0] = INDEX_NONE;
+								Parents[LOD][i][1] = INDEX_NONE;
+								Parents[LOD][i][2] = INDEX_NONE;
+								Parents[LOD][i][3] = INDEX_NONE;
+
+								Weights[LOD][i][0] = 0.0;
+								Weights[LOD][i][1] = 0.0;
+								Weights[LOD][i][2] = 0.0;
+								Weights[LOD][i][3] = 0.0;
+
+								Offsets[LOD][i][0] = 0.0;
+								Offsets[LOD][i][1] = 0.0;
+								Offsets[LOD][i][2] = 0.0;
+							}
+						}
+					} // end for all vertices
+				} // end for all LOD
+
+				// Stash bindings in the geometry collection
+				GeometryCollection::Facades::FTetrahedralBindings TetBindings(InCollection);
+				TetBindings.DefineSchema();
+				FName MeshName(*MeshId, MeshId.Len());
+				for (int32 LOD = 0; LOD < MeshVertices.Num(); LOD++)
+				{
+					TetBindings.AddBindingsGroup(TetMeshIdx, MeshName, LOD);
+					TetBindings.SetBindingsData(Parents[LOD], Weights[LOD], Offsets[LOD]);
+				}
+			} // end for TetMeshIdx
+		}
+		SetValue<DataType>(Context, InCollection, &Collection);
+	}
+}
diff --git a/Engine/Plugins/Experimental/ChaosFlesh/Source/ChaosFleshNodes/Private/Dataflow/ChaosFleshNodesPlugin.cpp b/Engine/Plugins/Experimental/ChaosFlesh/Source/ChaosFleshNodes/Private/Dataflow/ChaosFleshNodesPlugin.cpp
--- a/Engine/Plugins/Experimental/ChaosFlesh/Source/ChaosFleshNodes/Private/Dataflow/ChaosFleshNodesPlugin.cpp
+++ b/Engine/Plugins/Experimental/ChaosFlesh/Source/ChaosFleshNodes/Private/Dataflow/ChaosFleshNodesPlugin.cpp
@@ -5,8 +5,9 @@
 #include "CoreMinimal.h"
 
 #include "Dataflow/DataflowNodeFactory.h"
-#include "Dataflow/ChaosFleshEngineAssetNodes.h"
+#include "Dataflow/ChaosFleshBindingsNodes.h"
 #include "Dataflow/ChaosFleshCoreNodes.h"
+#include "Dataflow/ChaosFleshEngineAssetNodes.h"
 #include "Dataflow/ChaosFleshFiberDirectionInitializationNodes.h"
 #include "Dataflow/ChaosFleshKinematicInitializationNodes.h"
 #include "Dataflow/ChaosFleshRenderInitializationNodes.h"
@@ -24,6 +25,7 @@ void IChaosFleshNodesPlugin::StartupModule()
 	Dataflow::ChaosFleshRenderInitializationNodes();
 	Dataflow::RegisterChaosFleshKinematicInitializationNodes();
 	Dataflow::ChaosFleshTetrahedralNodes();
+	Dataflow::ChaosFleshBindingsNodes();
 }
 
 void IChaosFleshNodesPlugin::ShutdownModule()
diff --git a/Engine/Plugins/Experimental/ChaosFlesh/Source/ChaosFleshNodes/Public/Dataflow/ChaosFleshBindingsNodes.h b/Engine/Plugins/Experimental/ChaosFlesh/Source/ChaosFleshNodes/Public/Dataflow/ChaosFleshBindingsNodes.h
new file mode 100644
--- /dev/null
+++ b/Engine/Plugins/Experimental/ChaosFlesh/Source/ChaosFleshNodes/Public/Dataflow/ChaosFleshBindingsNodes.h
@@ -0,0 +1,51 @@
+// Copyright Epic Games, Inc. All Rights Reserved.
+
+#pragma once 
+
+#include "CoreMinimal.h"
+#include "Dataflow/DataflowEngine.h"
+#include "GeometryCollection/ManagedArrayCollection.h"
+
+#include "ChaosFleshBindingsNodes.generated.h"
+
+class UStaticMesh;
+class USkeletalMesh;
+
+USTRUCT(meta = (DataflowFlesh))
+struct FGenerateBindings : public FDataflowNode
+{
+	GENERATED_USTRUCT_BODY()
+	DATAFLOW_NODE_DEFINE_INTERNAL(FGenerateBindings, "GenerateBindings", "Flesh", "")
+	DATAFLOW_NODE_RENDER_TYPE(FGeometryCollection::StaticType(), "Collection")
+
+public:
+	typedef FManagedArrayCollection DataType;
+
+	UPROPERTY(meta = (DataflowInput, DataflowOutput, DisplayName = "Collection"))
+	FManagedArrayCollection Collection;
+
+	UPROPERTY(EditAnywhere, Category = "Dataflow")
+	TObjectPtr<const UStaticMesh> StaticMeshIn = nullptr;
+
+	UPROPERTY(EditAnywhere, Category = "Dataflow")
+	TObjectPtr<const USkeletalMesh> SkeletalMeshIn = nullptr;
+
+	UPROPERTY(EditAnywhere, Category = "Dataflow")
+	uint32 SurfaceProjectionIterations = 10;
+
+	FGenerateBindings(const Dataflow::FNodeParameters& InParam, FGuid InGuid = FGuid::NewGuid())
+		: FDataflowNode(InParam, InGuid)
+	{
+		RegisterInputConnection(&Collection);
+		RegisterOutputConnection(&Collection);
+		RegisterInputConnection(&StaticMeshIn);
+		RegisterInputConnection(&SkeletalMeshIn);
+	}
+
+	virtual void Evaluate(Dataflow::FContext& Context, const FDataflowOutput* Out) const override;
+};
+
+namespace Dataflow
+{
+	void ChaosFleshBindingsNodes();
+}
\ No newline at end of file
diff --git a/Engine/Source/Runtime/Experimental/Chaos/Private/Chaos/TriangleMesh.cpp b/Engine/Source/Runtime/Experimental/Chaos/Private/Chaos/TriangleMesh.cpp
--- a/Engine/Source/Runtime/Experimental/Chaos/Private/Chaos/TriangleMesh.cpp
+++ b/Engine/Source/Runtime/Experimental/Chaos/Private/Chaos/TriangleMesh.cpp
@@ -3,8 +3,8 @@
 
 #include "Chaos/Box.h"
 #include "Chaos/Defines.h"
-#include "Chaos/Math/Poisson.h"
 #include "Chaos/Plane.h"
+#include "Chaos/SmoothProject.h"
 #include "Chaos/Triangle.h"
 #include "Chaos/TriangleCollisionPoint.h"
 #include "HAL/IConsoleManager.h"
@@ -346,14 +346,14 @@ TArray<TVec3<T>> FTriangleMesh::GetFaceNormals(const TConstArrayView<TVec3<T>>&
 template CHAOS_API TArray<TVec3<FRealSingle>> FTriangleMesh::GetFaceNormals<FRealSingle>(const TConstArrayView<TVec3<FRealSingle>>&, const bool) const;
 template CHAOS_API TArray<TVec3<FRealDouble>> FTriangleMesh::GetFaceNormals<FRealDouble>(const TConstArrayView<TVec3<FRealDouble>>&, const bool) const;
 
-TArray<FVec3> FTriangleMesh::GetPointNormals(const TConstArrayView<FVec3>& Points, const bool ReturnEmptyOnError)
+TArray<FVec3> FTriangleMesh::GetPointNormals(const TConstArrayView<FVec3>& Points, const bool ReturnEmptyOnError, const bool bUseGlobalArray)
 {
 	TArray<FVec3> PointNormals;
 	const TArray<FVec3> FaceNormals = GetFaceNormals(Points, ReturnEmptyOnError);
 	if (FaceNormals.Num())
 	{
-		PointNormals.SetNumUninitialized(MNumIndices);
-		GetPointNormals(PointNormals, FaceNormals, /*bUseGlobalArray =*/ false);
+		PointNormals.SetNumUninitialized(bUseGlobalArray ? MNumIndices+MStartIdx : MNumIndices);
+		GetPointNormals(PointNormals, FaceNormals, bUseGlobalArray);
 	}
 	return PointNormals;
 }
@@ -1497,8 +1497,8 @@ void FTriangleMesh::BuildBVH(const TConstArrayView<TVec3<T>>& Points, TBVHType<T
 	}
 	BVH.Reinitialize(BVEntries);
 }
-template void FTriangleMesh::BuildBVH<FRealSingle>(const TConstArrayView<TVec3<FRealSingle>>& Points, TBVHType<FRealSingle>& BVH) const;
-template void FTriangleMesh::BuildBVH<FRealDouble>(const TConstArrayView<TVec3<FRealDouble>>& Points, TBVHType<FRealDouble>& BVH) const;
+template CHAOS_API void FTriangleMesh::BuildBVH<FRealSingle>(const TConstArrayView<TVec3<FRealSingle>>& Points, TBVHType<FRealSingle>& BVH) const;
+template CHAOS_API void FTriangleMesh::BuildBVH<FRealDouble>(const TConstArrayView<TVec3<FRealDouble>>& Points, TBVHType<FRealDouble>& BVH) const;
 
 template<typename T>
 bool FTriangleMesh::PointProximityQuery(const TBVHType<T>& BVH, const TConstArrayView<TVec3<T>>& Points, const int32 PointIndex, const TVec3<T>& PointPosition, const T PointThickness, const T ThisThickness, 
@@ -1597,168 +1597,7 @@ template bool FTriangleMesh::EdgeIntersectionQuery<FRealSingle>(const TBVHType<F
 template bool FTriangleMesh::EdgeIntersectionQuery<FRealDouble>(const TBVHType<FRealDouble>& BVH, const TConstArrayView<TVec3<FRealDouble>>& Points, const int32 EdgeIndex, const TVec3<FRealDouble>& EdgePosition1, const TVec3<FRealDouble>& EdgePosition2,
 	TFunctionRef<bool(const int32 EdgeIndex, const int32 TriangleIndex)> BroadphaseTest, TArray<TTriangleCollisionPoint<FRealDouble>>& Result) const;
 
-template <class T>
-bool SmoothProject(
-	const TConstArrayView<FVec3>& Points,
-	const TArray<TVec3<int32>>& Tris,
-	const TArray<FVec3>& PointNormals,
-	const FVec3& Pos,
-	const int32 TriIdx,
-	FVec3& Weights,
-	TArray<TArray<FVec3>>& TangentBases,
-	T* B,
-	T* C,
-	TVec3<T>& C0,
-	TVec3<T>& C1,
-	TVec3<T>& C2,
-	TVec3<T>& W0,
-	TVec3<T>& W1,
-	TVec3<T>& W2)
-{
-	// Construct orthogonal tangents from point normal
-	for (int32 i = 0; i < 3; i++)
-	{
-		const int32 PtIdx = Tris[TriIdx][i];
-		const FVec3& PointNormal = PointNormals[PtIdx];
-		TangentBases[i][0] = PointNormal.GetOrthogonalVector().GetSafeNormal();
-		TangentBases[i][1] = FVec3::CrossProduct(TangentBases[i][0], PointNormal);
-	}
-	// Solve
-	for (int32 i = 0; i < 3; i++)
-	{
-		for (int j = 0; j < 3; j++)
-		{
-			const int32 PtIdx = Tris[TriIdx][i];
-			const FVec3& Pt = Points[PtIdx];
-			const FVec3& TangentBase0 = TangentBases[i][0];
-			const FVec3& TangentBase1 = TangentBases[i][1];
-			FReal B0j = FVec3::DotProduct(TangentBase0, Pt);
-			FReal B1j = FVec3::DotProduct(TangentBase1, Pt);
-			RowMaj3x3Set(B, j, 0, static_cast<T>(B0j));
-			RowMaj3x3Set(B, j, 1, static_cast<T>(B1j));
-			RowMaj3x3Set(B, j, 2, static_cast<T>(1.0));
-		}
-		RowMaj3x3Inverse(B, B); // Sets B
-		FVec3 PV(
-			FVec3::DotProduct(TangentBases[i][0], Pos),
-			FVec3::DotProduct(TangentBases[i][1], Pos),
-			static_cast<T>(1.0));
-		RowMaj3x3SetRow(C, i, RowMaj3x3Multiply(B, PV));
-	}
-
-	RowMaj3x3Get(C, 0, 0) -= static_cast<T>(1.0);
-	RowMaj3x3Get(C, 1, 1) -= static_cast<T>(1.0);
-	RowMaj3x3Get(C, 2, 2) -= static_cast<T>(1.0);
-
-	RowMaj3x3GetRow(C, 0, C0);
-	RowMaj3x3GetRow(C, 1, C1);
-	RowMaj3x3GetRow(C, 2, C2);
-
-	W0 = FVec3::CrossProduct(C0, C1);
-	W1 = FVec3::CrossProduct(C0, C2);
-	W2 = FVec3::CrossProduct(C1, C2);
-
-	if (W0[0] >= 0.0 && W0[1] >= 0.0 && W0[2] >= 0.0)
-	{
-		Weights = W0;
-		Swap(Weights[0], Weights[2]);
-	}
-	else if (W1[0] >= 0.0 && W1[1] >= 0.0 && W1[2] >= 0.0)
-	{
-		Weights = W1;
-	}
-	else if (W2[0] >= 0.0 && W2[1] >= 0.0 && W2[2] >= 0.0)
-	{
-		Weights = W2;
-		Swap(Weights[0], Weights[2]);
-	}
-	else
-	{
-		// We're outside of the cone implied by the normals,
-		// so the projection doesn't exist.
-		return false;
-	}
-
-	FReal WeightsSum = Weights[0] + Weights[1] + Weights[2];
-	if (fabs(WeightsSum) < 1.0e-6)
-	{
-		// The point is at the vertex of the normal cone, so the
-		// projection isn't unique.  Any face location will do.
-		static const T OneThird = static_cast<T>(1.0 / 3);
-		Weights = FVec3(OneThird, OneThird, OneThird);
-		return true;
-	}
-	Weights /= WeightsSum;
-	return true;
-}
-
-template <class T>
-bool SmoothProject(
-	const TConstArrayView<FVec3>& Points,
-	const TArray<TVec3<int32>>& Tris,
-	const TArray<FVec3>& PointNormals,
-	const FVec3& Pos,
-	const int32 TriIdx,
-	FVec3& Weights)
-{
-	TArray<TArray<FVec3>> TangentBases;
-	TangentBases.SetNum(3);
-	for (int32 i = 0; i < 3; i++)
-	{
-		TangentBases[i].SetNum(2);
-	}
-	T B[9];
-	T C[9];
-	TVec3<T> C0;
-	TVec3<T> C1;
-	TVec3<T> C2;
-	TVec3<T> W0;
-	TVec3<T> W1;
-	TVec3<T> W2;
-
-	return SmoothProject(Points, Tris, PointNormals, Pos, TriIdx, Weights,
-		TangentBases, &B[0], &C[0], C0, C1, C2, W0, W1, W2);
-}
 
-template <class T>
-TArray<bool> SmoothProject(
-	const TConstArrayView<FVec3>& Points,
-	const TArray<TVec3<int32>>& Tris,
-	const TArray<FVec3>& PointNormals,
-	const FVec3& Pos,
-	const TArray<int32>& TriIdx,
-	TArray<FVec3>& Weights,
-	const bool UseFirstFound=false)
-{
-	TArray<TArray<FVec3>> TangentBases;
-	TangentBases.SetNum(3);
-	for (int32 i = 0; i < 3; i++)
-	{
-		TangentBases[i].SetNum(2);
-	}
-	T B[9];
-	T C[9];
-	TVec3<T> C0;
-	TVec3<T> C1;
-	TVec3<T> C2;
-	TVec3<T> W0;
-	TVec3<T> W1;
-	TVec3<T> W2;
-
-	TArray<bool> RetVal;
-	RetVal.SetNumZeroed(TriIdx.Num());
-	Weights.SetNum(TriIdx.Num());
-	for (int i = 0; i < TriIdx.Num(); i++)
-	{
-		RetVal[i] = SmoothProject(Points, Tris, PointNormals, Pos, TriIdx[i], Weights[i],
-			TangentBases, &B[0], &C[0], C0, C1, C2, W0, W1, W2);
-		if (UseFirstFound && RetVal[i])
-		{
-			return RetVal;
-		}
-	}
-	return RetVal;
-}
 
 template<typename T>
 bool FTriangleMesh::SmoothProject(
@@ -1770,33 +1609,47 @@ bool FTriangleMesh::SmoothProject(
 	FVec3& Weights, 
 	const int32 MaxIters) const
 {
-	TArray<int32> CandidateTris;
+	TSet<int32> SkipTris;
 	int32 Iter = 0;
 	do {
-		FAABB3 Box(Pos-FVec3(.5*Iter), Pos+FVec3(.5*Iter));
-		CandidateTris = BVH.FindAllIntersections(Box);
-	} while (!CandidateTris.Num() && ++Iter < MaxIters);
-
-	if (!CandidateTris.Num())
-	{
-		return false;
-	}
-
-	TArray<FVec3> CandidateWeights;
-	TArray<bool> Success = Chaos::SmoothProject<T>(Points, MElements, PointNormals, Pos, CandidateTris, CandidateWeights, true);
-	for (int32 i = 0; i < Success.Num(); i++)
-	{
-		if (Success[i])
+		TArray<int32> CandidateTris;
+		do {
+			FAABB3 Box(Pos - FVec3(.5 * Iter), Pos + FVec3(.5 * Iter));
+			CandidateTris = BVH.FindAllIntersections(Box);
+			// Remove candidates we've already considered
+			for (int32 i = CandidateTris.Num() - 1; i >= 0; i--)
+			{
+				if (SkipTris.Contains(CandidateTris[i]))
+				{
+					CandidateTris.RemoveAt(i);
+				}
+			}
+			++Iter;
+		} while (!CandidateTris.Num() && Iter < MaxIters);
+		if (!CandidateTris.Num())
 		{
-			TriangleIndex = CandidateTris[i];
-			Weights = CandidateWeights[i];
-			return true;
+			// If we exited the BVH loop with no candidates, then we've exceeded MaxIters.
+			return false;
 		}
-	}
+		// Test candidates
+		TArray<FVec3> CandidateWeights;
+		TArray<bool> Success = Chaos::SmoothProject<T>(Points, MElements, PointNormals, Pos, CandidateTris, CandidateWeights, true);
+		for (int32 i = 0; i < Success.Num(); i++)
+		{
+			if (Success[i])
+			{
+				TriangleIndex = CandidateTris[i];
+				Weights = CandidateWeights[i];
+				return true;
+			}
+		}
+		// No hits. Add candidates to skip list, and go searching for more.
+		SkipTris.Append(CandidateTris);
+	} while (Iter < MaxIters);
 	return false;
 }
-template bool FTriangleMesh::SmoothProject<FRealSingle>(const TBVHType<FRealSingle>& BVH, const TConstArrayView<FVec3>& Points, const TArray<FVec3>& PointNormals, const FVec3& Point, int32& TriangleIndex, FVec3& Weights, const int32 MaxIters) const;
-template bool FTriangleMesh::SmoothProject<FRealDouble>(const TBVHType<FRealDouble>& BVH, const TConstArrayView<FVec3>& Points, const TArray<FVec3>& PointNormals, const FVec3& Point, int32& TriangleIndex, FVec3& Weights, const int32 MaxIters) const;
+template CHAOS_API bool FTriangleMesh::SmoothProject<FRealSingle>(const TBVHType<FRealSingle>& BVH, const TConstArrayView<FVec3>& Points, const TArray<FVec3>& PointNormals, const FVec3& Point, int32& TriangleIndex, FVec3& Weights, const int32 MaxIters) const;
+template CHAOS_API bool FTriangleMesh::SmoothProject<FRealDouble>(const TBVHType<FRealDouble>& BVH, const TConstArrayView<FVec3>& Points, const TArray<FVec3>& PointNormals, const FVec3& Point, int32& TriangleIndex, FVec3& Weights, const int32 MaxIters) const;
 
 template<typename T>
 void FTriangleMesh::BuildSpatialHash(const TConstArrayView<TVec3<T>>& Points, TSpatialHashType<T>& SpatialHash) const
diff --git a/Engine/Source/Runtime/Experimental/Chaos/Private/GeometryCollection/Facades/CollectionTetrahedralBindingsFacade.cpp b/Engine/Source/Runtime/Experimental/Chaos/Private/GeometryCollection/Facades/CollectionTetrahedralBindingsFacade.cpp
new file mode 100644
--- /dev/null
+++ b/Engine/Source/Runtime/Experimental/Chaos/Private/GeometryCollection/Facades/CollectionTetrahedralBindingsFacade.cpp
@@ -0,0 +1,222 @@
+// Copyright Epic Games, Inc. All Rights Reserved.
+
+/*=============================================================================
+	GeometryCollection.cpp: FGeometryCollection methods.
+=============================================================================*/
+
+#include "GeometryCollection/Facades/CollectionTetrahedralBindingsFacade.h"
+#include "GeometryCollection/GeometryCollection.h"
+#include "GeometryCollection/ManagedArrayAccessor.h"
+
+namespace GeometryCollection::Facades
+{
+	// Groups 
+	const FName FTetrahedralBindings::MeshBindingsGroupName = "MeshBindings";
+
+	// Attributes
+	const FName FTetrahedralBindings::MeshIdAttributeName = "MeshId";
+	//const FName FTetrahedralBindings::MeshLODAttributeName = "MeshLOD";
+
+	const FName FTetrahedralBindings::ParentsAttributeName = "Parents";
+	const FName FTetrahedralBindings::WeightsAttributeName = "Weights";
+	const FName FTetrahedralBindings::OffsetsAttributeName = "Offsets";
+
+	FTetrahedralBindings::FTetrahedralBindings(FManagedArrayCollection& InCollection)
+		: MeshIdAttribute(InCollection, MeshIdAttributeName, MeshBindingsGroupName)
+		, Parents(nullptr)
+		, Weights(nullptr)
+		, Offsets(nullptr)
+	{}
+
+	FTetrahedralBindings::FTetrahedralBindings(const FManagedArrayCollection& InCollection)
+		: MeshIdAttribute(InCollection, MeshIdAttributeName, MeshBindingsGroupName)
+		, Parents(nullptr)
+		, Weights(nullptr)
+		, Offsets(nullptr)
+	{}
+
+	FTetrahedralBindings::~FTetrahedralBindings()
+	{
+		delete Parents; Parents = nullptr;
+		delete Weights; Weights = nullptr;
+		delete Offsets; Offsets = nullptr;
+	}
+
+	void FTetrahedralBindings::DefineSchema(/*const FName* MeshId, const int32 LOD*/)
+	{
+		check(!IsConst());
+		TManagedArray<FString>& MeshIdValues = 
+			MeshIdAttribute.IsValid() ? MeshIdAttribute.Modify() : MeshIdAttribute.Add();
+	/*	if (MeshId)
+		{
+			// Generate the group name, and if it doesn't exist, store it in the table of 
+			// available bindings groups.
+			FName GroupName = GenerateMeshGroupName(*MeshId, LOD);
+			if (!MeshIdValues.Contains(GroupName.ToString()))
+			{
+				const int32 Idx = MeshIdAttribute.AddElements(1);
+				MeshIdValues[Idx] = GroupName.ToString();
+
+				// This is a new group, so create the bindings arrays.
+				AddBindingsGroup(GroupName);
+			}
+			else
+			{
+				// This is an existing group, so find the existing bindings arrays.
+				ReadBindingsGroup(GroupName);
+			}
+
+		}
+		*/
+	}
+
+	bool FTetrahedralBindings::IsValid() const
+	{
+		return MeshIdAttribute.IsValid() && 
+			(Parents && Parents->IsValid()) && 
+			(Weights && Weights->IsValid()) && 
+			(Offsets && Offsets->IsValid());
+	}
+
+	FName FTetrahedralBindings::GenerateMeshGroupName(
+		const int32 TetMeshIdx,
+		const FName& MeshId,
+		const int32 LOD)
+	{
+		FString Str = FString::Printf(TEXT("TetrahedralBindings:TetMeshIdx:%d:%s:%d"), TetMeshIdx, *MeshId.ToString(), LOD);
+		return FName(Str.Len(), *Str);
+	}
+
+	bool FTetrahedralBindings::ContainsBindingsGroup(const int32 TetMeshIdx, const FName& MeshId, const int32 LOD) const
+	{
+		return ContainsBindingsGroup(GenerateMeshGroupName(TetMeshIdx, MeshId, LOD));
+	}
+
+	bool FTetrahedralBindings::ContainsBindingsGroup(const FName& GroupName) const
+	{
+		check(MeshIdAttribute.IsValid());
+		const TManagedArray<FString>* MeshIdValues =
+			MeshIdAttribute.IsValid() ? MeshIdAttribute.Find() : nullptr;
+		return MeshIdValues ? MeshIdValues->Contains(GroupName.ToString()) : false;
+	}
+
+	void FTetrahedralBindings::AddBindingsGroup(const int32 TetMeshIdx, const FName& MeshId, const int32 LOD)
+	{
+		AddBindingsGroup(GenerateMeshGroupName(TetMeshIdx, MeshId, LOD));
+	}
+
+	void FTetrahedralBindings::AddBindingsGroup(const FName& GroupName)
+	{
+		if (ContainsBindingsGroup(GroupName))
+		{
+			ReadBindingsGroup(GroupName);
+			return;
+		}
+
+		check(!IsConst());
+		const int32 Idx = MeshIdAttribute.AddElements(1);
+		MeshIdAttribute.Modify()[Idx] = GroupName.ToString();
+
+		delete Parents; Parents = nullptr;
+		delete Weights; Weights = nullptr;
+		delete Offsets; Offsets = nullptr;
+		FManagedArrayCollection& Collection = *MeshIdAttribute.GetCollection();
+		Parents = new TManagedArrayAccessor<FIntVector4>(Collection, ParentsAttributeName, GroupName);
+		Weights = new TManagedArrayAccessor<FVector4f>(Collection, WeightsAttributeName, GroupName);
+		Offsets = new TManagedArrayAccessor<FVector3f>(Collection, OffsetsAttributeName, GroupName);
+		Parents->Add(ManageArrayAccessor::EPersistencePolicy::MakePersistent, FGeometryCollection::VerticesGroup);
+		Weights->Add(ManageArrayAccessor::EPersistencePolicy::MakePersistent);
+		Offsets->Add(ManageArrayAccessor::EPersistencePolicy::MakePersistent);
+	}
+
+	bool FTetrahedralBindings::ReadBindingsGroup(const int32 TetMeshIdx, const FName& MeshId, const int32 LOD)
+	{
+		return ReadBindingsGroup(GenerateMeshGroupName(TetMeshIdx, MeshId, LOD));
+	}
+
+	bool FTetrahedralBindings::ReadBindingsGroup(const FName& GroupName)
+	{
+		check(MeshIdAttribute.IsValid());
+		delete Parents; Parents = nullptr;
+		delete Weights; Weights = nullptr;
+		delete Offsets; Offsets = nullptr;
+		if (MeshIdAttribute.Find()->Contains(GroupName.ToString()))
+		{
+			return false;
+		}
+		// This is an existing group, so find the existing bindings arrays.
+		FManagedArrayCollection& Collection = *MeshIdAttribute.GetCollection();
+		Parents = new TManagedArrayAccessor<FIntVector4>(Collection, ParentsAttributeName, GroupName);
+		Weights = new TManagedArrayAccessor<FVector4f>(Collection, WeightsAttributeName, GroupName);
+		Offsets = new TManagedArrayAccessor<FVector3f>(Collection, OffsetsAttributeName, GroupName);
+		check(Parents->IsValid());
+		check(Weights->IsValid());
+		check(Offsets->IsValid());
+		return Parents->IsValid() && Weights->IsValid() && Offsets->IsValid();
+	}
+
+	void FTetrahedralBindings::RemoveBindingsGroup(const int32 TetMeshIdx, const FName& MeshId, const int32 LOD)
+	{
+		RemoveBindingsGroup(GenerateMeshGroupName(TetMeshIdx, MeshId, LOD));
+	}
+
+	void FTetrahedralBindings::RemoveBindingsGroup(const FName& GroupName)
+	{
+		check(!IsConst());
+		TManagedArray<FString>& MeshIdValues = MeshIdAttribute.Modify();
+		int32 Idx = MeshIdValues.Find(GroupName.ToString());
+		if (Idx != INDEX_NONE)
+		{
+			TArray<int32> Indices;
+			Indices.Add(Idx);
+			MeshIdValues.RemoveElements(Indices);
+		}
+
+		FManagedArrayCollection& Collection = *MeshIdAttribute.GetCollection();
+		if (Parents)
+		{
+			Parents->Remove();
+		}
+		if (Weights)
+		{
+			Weights->Remove();
+		}
+		if (Offsets)
+		{
+			Offsets->Remove();
+		}
+		// Only drop the group if it's empty at this point?
+		if (Collection.NumAttributes(GroupName) == 0)
+		{
+			Collection.RemoveGroup(GroupName);
+		}
+	}
+
+	void 
+	FTetrahedralBindings::SetBindingsData(
+		const TArray<FIntVector4>& ParentsIn,
+		const TArray<FVector4f>& WeightsIn,
+		const TArray<FVector3f>& OffsetsIn)
+	{
+		check(!IsConst());
+		check(IsValid());
+		check((ParentsIn.Num() == WeightsIn.Num()) && (ParentsIn.Num() == OffsetsIn.Num()));
+
+		const int32 Num = ParentsIn.Num();
+		const int32 CurrNum = Parents->Num();//Collection.NumElements(CurrGroupName);
+		Parents->AddElements(Num - CurrNum); // Resizes the group
+		TManagedArray<FIntVector4>& ParentsValues = Parents->Modify();
+		TManagedArray<FVector4f>& WeightsValues = Weights->Modify();
+		TManagedArray<FVector3f>& OffsetsValues = Offsets->Modify();
+		for (int32 i = 0; i < Num; i++)
+		{
+			ParentsValues[i] = ParentsIn[i];
+			WeightsValues[i] = WeightsIn[i];
+			OffsetsValues[i] = OffsetsIn[i];
+		}
+	}
+
+
+};
+
+
diff --git a/Engine/Source/Runtime/Experimental/Chaos/Private/GeometryCollection/ManagedArrayCollection.cpp b/Engine/Source/Runtime/Experimental/Chaos/Private/GeometryCollection/ManagedArrayCollection.cpp
--- a/Engine/Source/Runtime/Experimental/Chaos/Private/GeometryCollection/ManagedArrayCollection.cpp
+++ b/Engine/Source/Runtime/Experimental/Chaos/Private/GeometryCollection/ManagedArrayCollection.cpp
@@ -27,6 +27,19 @@ void FManagedArrayCollection::AddGroup(FName Group)
 	GroupInfo.Add(Group, info);
 }
 
+int32 FManagedArrayCollection::NumAttributes(FName Group) const
+{
+	int32 Num=0;
+	for (const TTuple<FKeyType, FValueType>& Entry : Map)
+	{
+		if (Entry.Key.Get<1>() == Group)
+		{
+			Num++;
+		}
+	}
+	return Num;
+}
+
 void FManagedArrayCollection::RemoveElements(const FName& Group, const TArray<int32>& SortedDeletionList, FProcessingParameters Params)
 {
 	if (SortedDeletionList.Num())
diff --git a/Engine/Source/Runtime/Experimental/Chaos/Public/Chaos/SmoothProject.h b/Engine/Source/Runtime/Experimental/Chaos/Public/Chaos/SmoothProject.h
new file mode 100644
--- /dev/null
+++ b/Engine/Source/Runtime/Experimental/Chaos/Public/Chaos/SmoothProject.h
@@ -0,0 +1,170 @@
+// Copyright Epic Games, Inc. All Rights Reserved.
+#pragma once
+
+#include "Chaos/Defines.h"
+#include "Chaos/Math/Poisson.h"
+
+namespace Chaos {
+
+	template <class T>
+	bool SmoothProject(
+		const TConstArrayView<FVec3>& Points,
+		const TArray<TVec3<int32>>& Tris,
+		const TArray<FVec3>& PointNormals,
+		const FVec3& Pos,
+		const int32 TriIdx,
+		FVec3& Weights,
+		TArray<TArray<FVec3>>& TangentBases,
+		T* B,
+		T* C,
+		TVec3<T>& C0,
+		TVec3<T>& C1,
+		TVec3<T>& C2,
+		TVec3<T>& W0,
+		TVec3<T>& W1,
+		TVec3<T>& W2)
+	{
+		// Construct orthogonal tangents from point normal
+		for (int32 i = 0; i < 3; i++)
+		{
+			const int32 PtIdx = Tris[TriIdx][i];
+			const FVec3& PointNormal = PointNormals[PtIdx];
+			TangentBases[i][0] = PointNormal.GetOrthogonalVector().GetSafeNormal();
+			TangentBases[i][1] = FVec3::CrossProduct(TangentBases[i][0], PointNormal);
+		}
+		// Solve
+		for (int32 i = 0; i < 3; i++)
+		{
+			for (int j = 0; j < 3; j++)
+			{
+				const int32 PtIdx = Tris[TriIdx][j];
+				const FVec3& Pt = Points[PtIdx];
+				const FVec3& TangentBase0 = TangentBases[i][0];
+				const FVec3& TangentBase1 = TangentBases[i][1];
+				FReal B0j = FVec3::DotProduct(TangentBase0, Pt);
+				FReal B1j = FVec3::DotProduct(TangentBase1, Pt);
+				RowMaj3x3Set(B, 0, j, static_cast<T>(B0j));
+				RowMaj3x3Set(B, 1, j, static_cast<T>(B1j));
+				RowMaj3x3Set(B, 2, j, static_cast<T>(1.0));
+			}
+			RowMaj3x3Inverse(B, B); // Sets B
+			FVec3 PV(
+				FVec3::DotProduct(TangentBases[i][0], Pos),
+				FVec3::DotProduct(TangentBases[i][1], Pos),
+				static_cast<T>(1.0));
+			RowMaj3x3SetCol(C, i, RowMaj3x3Multiply(B, PV));
+		}
+
+		RowMaj3x3Get(C, 0, 0) -= static_cast<T>(1.0);
+		RowMaj3x3Get(C, 1, 1) -= static_cast<T>(1.0);
+		RowMaj3x3Get(C, 2, 2) -= static_cast<T>(1.0);
+
+		RowMaj3x3GetRow(C, 0, C0);
+		RowMaj3x3GetRow(C, 1, C1);
+		RowMaj3x3GetRow(C, 2, C2);
+
+		W0 = FVec3::CrossProduct(C0, C1);
+		W1 = FVec3::CrossProduct(C0, C2);
+		W2 = FVec3::CrossProduct(C1, C2);
+
+		if (W0[0] >= 0.0 && W0[1] >= 0.0 && W0[2] >= 0.0)
+		{
+			Weights = W0;
+		}
+		else if (W1[0] >= 0.0 && W1[1] >= 0.0 && W1[2] >= 0.0)
+		{
+			Weights = W1;
+		}
+		else if (W2[0] >= 0.0 && W2[1] >= 0.0 && W2[2] >= 0.0)
+		{
+			Weights = W2;
+		}
+		else
+		{
+			// We're outside of the cone implied by the normals,
+			// so the projection doesn't exist.
+			return false;
+		}
+
+		FReal WeightsSum = Weights[0] + Weights[1] + Weights[2];
+		if (fabs(WeightsSum) < 1.0e-6)
+		{
+			// The point is at the vertex of the normal cone, so the
+			// projection isn't unique.  Any face location will do.
+			static const T OneThird = static_cast<T>(1.0 / 3);
+			Weights = FVec3(OneThird, OneThird, OneThird);
+			return true;
+		}
+		Weights /= WeightsSum;
+		return true;
+	}
+
+	template <class T>
+	bool SmoothProject(
+		const TConstArrayView<FVec3>& Points,
+		const TArray<TVec3<int32>>& Tris,
+		const TArray<FVec3>& PointNormals,
+		const FVec3& Pos,
+		const int32 TriIdx,
+		FVec3& Weights)
+	{
+		TArray<TArray<FVec3>> TangentBases;
+		TangentBases.SetNum(3);
+		for (int32 i = 0; i < 3; i++)
+		{
+			TangentBases[i].SetNum(2);
+		}
+		T B[9];
+		T C[9];
+		TVec3<T> C0;
+		TVec3<T> C1;
+		TVec3<T> C2;
+		TVec3<T> W0;
+		TVec3<T> W1;
+		TVec3<T> W2;
+
+		return SmoothProject(Points, Tris, PointNormals, Pos, TriIdx, Weights,
+			TangentBases, &B[0], &C[0], C0, C1, C2, W0, W1, W2);
+	}
+
+	template <class T>
+	TArray<bool> SmoothProject(
+		const TConstArrayView<FVec3>& Points,
+		const TArray<TVec3<int32>>& Tris,
+		const TArray<FVec3>& PointNormals,
+		const FVec3& Pos,
+		const TArray<int32>& TriIdx,
+		TArray<FVec3>& Weights,
+		const bool UseFirstFound = false)
+	{
+		TArray<TArray<FVec3>> TangentBases;
+		TangentBases.SetNum(3);
+		for (int32 i = 0; i < 3; i++)
+		{
+			TangentBases[i].SetNum(2);
+		}
+		T B[9];
+		T C[9];
+		TVec3<T> C0;
+		TVec3<T> C1;
+		TVec3<T> C2;
+		TVec3<T> W0;
+		TVec3<T> W1;
+		TVec3<T> W2;
+
+		TArray<bool> RetVal;
+		RetVal.SetNumZeroed(TriIdx.Num());
+		Weights.SetNum(TriIdx.Num());
+		for (int i = 0; i < TriIdx.Num(); i++)
+		{
+			RetVal[i] = SmoothProject(Points, Tris, PointNormals, Pos, TriIdx[i], Weights[i],
+				TangentBases, &B[0], &C[0], C0, C1, C2, W0, W1, W2);
+			if (UseFirstFound && RetVal[i])
+			{
+				return RetVal;
+			}
+		}
+		return RetVal;
+	}
+
+} // namespace Chaos
\ No newline at end of file
diff --git a/Engine/Source/Runtime/Experimental/Chaos/Public/Chaos/TriangleMesh.h b/Engine/Source/Runtime/Experimental/Chaos/Public/Chaos/TriangleMesh.h
--- a/Engine/Source/Runtime/Experimental/Chaos/Public/Chaos/TriangleMesh.h
+++ b/Engine/Source/Runtime/Experimental/Chaos/Public/Chaos/TriangleMesh.h
@@ -153,7 +153,7 @@ namespace Chaos
 		FORCEINLINE TArray<FVec3> GetFaceNormals(const FParticles& InParticles, const bool ReturnEmptyOnError = true) const
 		{ return GetFaceNormals(TConstArrayView<FVec3>(InParticles.X()), ReturnEmptyOnError); }
 
-		CHAOS_API TArray<FVec3> GetPointNormals(const TConstArrayView<FVec3>& points, const bool ReturnEmptyOnError = true);
+		CHAOS_API TArray<FVec3> GetPointNormals(const TConstArrayView<FVec3>& points, const bool ReturnEmptyOnError = true, const bool bUseGlobalArray=false);
 		FORCEINLINE TArray<FVec3> GetPointNormals(const FParticles& InParticles, const bool ReturnEmptyOnError = true)
 		{ return GetPointNormals(TConstArrayView<FVec3>(InParticles.X()), ReturnEmptyOnError); }
 		
diff --git a/Engine/Source/Runtime/Experimental/Chaos/Public/GeometryCollection/Facades/CollectionTetrahedralBindingsFacade.h b/Engine/Source/Runtime/Experimental/Chaos/Public/GeometryCollection/Facades/CollectionTetrahedralBindingsFacade.h
new file mode 100644
--- /dev/null
+++ b/Engine/Source/Runtime/Experimental/Chaos/Public/GeometryCollection/Facades/CollectionTetrahedralBindingsFacade.h
@@ -0,0 +1,122 @@
+// Copyright Epic Games, Inc. All Rights Reserved.
+
+#pragma once
+#include "Containers/UnrealString.h"
+#include "Containers/Array.h"
+#include "GeometryCollection/ManagedArrayCollection.h"
+#include "GeometryCollection/ManagedArrayAccessor.h"
+
+namespace GeometryCollection::Facades
+{
+
+	/**
+	* FTetrahedralBindings
+	* 
+	* Interface for storing and retrieving bindings of surfaces (typically SkeletalMesh or StaticMesh) to
+	* tetrahedral meshes.  Bindings data for each surface is grouped by a mesh id and a level of detail.
+	*/
+	class CHAOS_API FTetrahedralBindings
+ 	{
+	public:
+
+		// groups
+		static const FName MeshBindingsGroupName;
+
+		// Attributes
+		static const FName MeshIdAttributeName;
+
+		static const FName ParentsAttributeName;
+		static const FName WeightsAttributeName;
+		static const FName OffsetsAttributeName;
+
+		/**
+		* FSelectionFacade Constuctor
+		* @param VertixDependencyGroup : GroupName the index attribute is dependent on. 
+		*/
+		FTetrahedralBindings(FManagedArrayCollection& InSelf);
+		FTetrahedralBindings(const FManagedArrayCollection& InSelf);
+		virtual ~FTetrahedralBindings();
+
+		/** 
+		* Create the facade schema. 
+		*/
+		void DefineSchema();
+
+		/** Returns \c true if the facade is operating on a read-only geometry collection. */
+		bool IsConst() const { return MeshIdAttribute.IsConst(); }
+
+		/** 
+		* Returns \c true if the Facade defined on the collection, and is initialized to
+		* a valid bindings group.
+		*/
+		bool IsValid() const;
+
+		/**
+		* Given a \p MeshId (by convention \code Mesh->GetPrimaryAssetId() \endcode) and
+		* a Level Of Detail rank, generate the associated bindings group name.
+		*/
+		static FName GenerateMeshGroupName(const int32 TetMeshIdx, const FName& MeshId, const int32 LOD);
+
+		/**
+		* Returns \c true if the specified bindings group exists.
+		*/
+		bool ContainsBindingsGroup(const int32 TetMeshIdx, const FName& MeshId, const int32 LOD) const;
+		bool ContainsBindingsGroup(const FName& GroupName) const;
+
+		/**
+		* Create a new bindings group, allocating new arrays.
+		*/
+		void AddBindingsGroup(const int32 TetMeshIdx, const FName& MeshId, const int32 LOD);
+		void AddBindingsGroup(const FName& GroupName);
+		
+		/**
+		* Initialize local arrays to point at a bindings group associated with \p MeshId 
+		* and \p LOD.  Returns \c false if it doesn't exist.
+		*/
+		bool ReadBindingsGroup(const int32 TetMeshIdx, const FName& MeshId, const int32 LOD);
+		bool ReadBindingsGroup(const FName& GroupName);
+
+		/**
+		* Removes a group from the list of bindings groups, removes the bindings arrays 
+		* from the geometry collection, and removes the group if it's empty.
+		*/
+		void RemoveBindingsGroup(const int32 TetMeshIdx, const FName& MeshId, const int32 LOD);
+		void RemoveBindingsGroup(const FName& GroupName);
+
+		/**
+		* Authors bindings data.
+		* 
+		* \p Parents are indicies of vertices; tetrahedron or surface triangle where final 
+		*    elem is \c INDEX_NONE.
+		* \p Weights are barycentric coordinates.
+		* \p Offsets are vectors from the barycentric point to the location, in the case 
+		*    of a surface binding.
+		*/
+		void SetBindingsData(const TArray<FIntVector4>& Parents, const TArray<FVector4f>& Weights, const TArray<FVector3f>& Offsets);
+
+		/**
+		* Get Parents array.
+		*/
+		const TManagedArrayAccessor<FIntVector4>* GetParents() const { return Parents; }
+		      TManagedArrayAccessor<FIntVector4>* GetParents() { check(!IsConst()); return Parents; }
+
+		/**
+		* Get Weights array.
+		*/
+		const TManagedArrayAccessor<FVector4f>* GetWeights() const { return Weights; }
+		      TManagedArrayAccessor<FVector4f>* GetWeights() { check(!IsConst()); return Weights; }
+
+		/**
+		* Get Offsets array.
+		*/
+		const TManagedArrayAccessor<FVector3f>* GetOffsets() const { return Offsets; }
+		      TManagedArrayAccessor<FVector3f>* GetOffsets() { check(!IsConst());  return Offsets; }
+
+	private:
+		TManagedArrayAccessor<FString> MeshIdAttribute;
+		TManagedArrayAccessor<FIntVector4>* Parents;
+		TManagedArrayAccessor<FVector4f>* Weights;
+		TManagedArrayAccessor<FVector3f>* Offsets;
+	};
+
+}
diff --git a/Engine/Source/Runtime/Experimental/Chaos/Public/GeometryCollection/ManagedArrayAccessor.h b/Engine/Source/Runtime/Experimental/Chaos/Public/GeometryCollection/ManagedArrayAccessor.h
--- a/Engine/Source/Runtime/Experimental/Chaos/Public/GeometryCollection/ManagedArrayAccessor.h
+++ b/Engine/Source/Runtime/Experimental/Chaos/Public/GeometryCollection/ManagedArrayAccessor.h
@@ -39,6 +39,15 @@ public:
 		, ConstAttributeArray(InCollection.FindAttributeTyped<T>(InAttributeName, InAttributeGroup))
 	{}
 
+	const FManagedArrayCollection& GetConstCollection() const { return ConstCollection; }
+	FManagedArrayCollection* GetCollection() { check(!IsConst()); return Collection; }
+
+	FName GetName() const { return Name; }
+
+	FName GetGroup() const { return Group; }
+
+	FName GetGroupDependency() const { return GroupDependency; }
+
 	bool IsConst() const { return Collection == nullptr; }
 
 	bool IsValid() const { return ConstAttributeArray != nullptr; }
diff --git a/Engine/Source/Runtime/Experimental/Chaos/Public/GeometryCollection/ManagedArrayCollection.h b/Engine/Source/Runtime/Experimental/Chaos/Public/GeometryCollection/ManagedArrayCollection.h
--- a/Engine/Source/Runtime/Experimental/Chaos/Public/GeometryCollection/ManagedArrayCollection.h
+++ b/Engine/Source/Runtime/Experimental/Chaos/Public/GeometryCollection/ManagedArrayCollection.h
@@ -223,6 +223,10 @@ public:
 	*/
 	void AddGroup(FName Group);
 
+	/**
+	* Returns the number of attributes in a group.
+	*/
+	int32 NumAttributes(FName Group) const;
 
 	/**
 	* List all the attributes in a group names.
@@ -395,7 +399,6 @@ public:
 	*/
 	void RemoveGroup(FName Group);
 
-
 	/**
 	* List all the group names.
 	*/
