Flesh : Added GeometryProcessing plugin to ChaosFlesh

diff --git a/Engine/Plugins/Experimental/ChaosFlesh/ChaosFlesh.uplugin b/Engine/Plugins/Experimental/ChaosFlesh/ChaosFlesh.uplugin
--- a/Engine/Plugins/Experimental/ChaosFlesh/ChaosFlesh.uplugin
+++ b/Engine/Plugins/Experimental/ChaosFlesh/ChaosFlesh.uplugin
@@ -49,6 +49,10 @@
 		{
 			"Name": "ChaosCaching",
 			"Enabled": true
+		},
+		{
+			"Name": "GeometryProcessing",
+			"Enabled": true
 		}
 	]
 }
