Chaos : AssetFactory icon setup

diff --git a/Engine/Plugins/Experimental/ChaosFlesh/Source/ChaosFleshEditor/Private/ChaosFlesh/Editor/FleshEditorStyle.h b/Engine/Plugins/Experimental/ChaosFlesh/Source/ChaosFleshEditor/Private/ChaosFlesh/Editor/FleshEditorStyle.h
--- a/Engine/Plugins/Experimental/ChaosFlesh/Source/ChaosFleshEditor/Private/ChaosFlesh/Editor/FleshEditorStyle.h
+++ b/Engine/Plugins/Experimental/ChaosFlesh/Source/ChaosFleshEditor/Private/ChaosFlesh/Editor/FleshEditorStyle.h
@@ -17,8 +17,8 @@ public:
 		SetContentRoot(IPluginManager::Get().FindPlugin("ChaosFlesh")->GetBaseDir() / TEXT("Resources"));
 		Set("ClassIcon.ChaosDeformableSolver", new FSlateVectorImageBrush(RootToContentDir(TEXT("ChaosSolver_16.svg")), Icon16x16));
 		Set("ClassThumbnail.ChaosDeformableSolver", new FSlateVectorImageBrush(RootToContentDir(TEXT("ChaosSolver_64.svg")), Icon64x64));
-		Set("ClassIcon.FleshAsset", new FSlateVectorImageBrush(RootToContentDir(TEXT("FleshAsset_16.svg")), Icon16x16));
-		Set("ClassThumbnail.FleshAsset", new FSlateVectorImageBrush(RootToContentDir(TEXT("FleshAsset_64.svg")), Icon64x64));
+		Set("ClassIcon.FleshAsset", new FSlateImageBrush(RootToContentDir(TEXT("FleshAsset_16x.png")), Icon16x16));
+		Set("ClassThumbnail.FleshAsset", new FSlateImageBrush(RootToContentDir(TEXT("FleshAsset_16x.png")), Icon64x64));
 		FSlateStyleRegistry::RegisterSlateStyle(*this);
 	}
 
