Flesh : StaticMesh connection to FGenerateTetrahedralCollectionDataflowNodes

diff --git a/Engine/Plugins/Experimental/ChaosFlesh/Source/ChaosFleshNodes/Private/Dataflow/ChaosFleshTetrahedralNodes.cpp b/Engine/Plugins/Experimental/ChaosFlesh/Source/ChaosFleshNodes/Private/Dataflow/ChaosFleshTetrahedralNodes.cpp
--- a/Engine/Plugins/Experimental/ChaosFlesh/Source/ChaosFleshNodes/Private/Dataflow/ChaosFleshTetrahedralNodes.cpp
+++ b/Engine/Plugins/Experimental/ChaosFlesh/Source/ChaosFleshNodes/Private/Dataflow/ChaosFleshTetrahedralNodes.cpp
@@ -135,14 +135,15 @@ void FGenerateTetrahedralCollectionDataflowNodes::EvaluateIsoStuffing(
 	TObjectPtr<const UStaticMesh>& InStaticMesh) const
 {
 #if WITH_EDITORONLY_DATA
-	if (InStaticMesh && NumCells > 0 && (-.5 <= OffsetPercent && OffsetPercent <= 0.5))
+	const TObjectPtr<const UStaticMesh>& StaticMeshValue = GetValue<TObjectPtr<const UStaticMesh>>(Context, &InStaticMesh);
+	if (StaticMeshValue && NumCells > 0 && (-.5 <= OffsetPercent && OffsetPercent <= 0.5))
 	{
 		// make a mesh description for UE::Geometry tools
 		UE::Geometry::FDynamicMesh3 DynamicMesh;
 		FMeshDescriptionToDynamicMesh GetSourceMesh;
-		bool bUsingHiResSource = InStaticMesh->IsHiResMeshDescriptionValid();
+		bool bUsingHiResSource = StaticMeshValue->IsHiResMeshDescriptionValid();
 		const FMeshDescription* UseSourceMeshDescription =
-			(bUsingHiResSource) ? InStaticMesh->GetHiResMeshDescription() : InStaticMesh->GetMeshDescription(0);
+			(bUsingHiResSource) ? StaticMeshValue->GetHiResMeshDescription() : StaticMeshValue->GetMeshDescription(0);
 		GetSourceMesh.Convert(UseSourceMeshDescription, DynamicMesh);
 
 		// Tet mesh generation
