ChaosFlesh : KinematicNode update.

diff --git a/Engine/Plugins/Experimental/ChaosFlesh/Source/ChaosFleshNodes/Private/Dataflow/ChaosFleshKinematicInitializationNodes.cpp b/Engine/Plugins/Experimental/ChaosFlesh/Source/ChaosFleshNodes/Private/Dataflow/ChaosFleshKinematicInitializationNodes.cpp
--- a/Engine/Plugins/Experimental/ChaosFlesh/Source/ChaosFleshNodes/Private/Dataflow/ChaosFleshKinematicInitializationNodes.cpp
+++ b/Engine/Plugins/Experimental/ChaosFlesh/Source/ChaosFleshNodes/Private/Dataflow/ChaosFleshKinematicInitializationNodes.cpp
@@ -134,51 +134,60 @@ void FKinematicInitializationDataflowNode::Evaluate(Dataflow::FContext& Context,
 {
 	if (Out->IsA<DataType>(&Collection))
 	{
+		TArray<int32> BoundVerts;
+		TArray<float> BoundWeights;
+
 		DataType InCollection = GetValue<DataType>(Context, &Collection);
 
-		if (TObjectPtr<USkeletalMesh> SkeletalMesh = GetValue<TObjectPtr<USkeletalMesh>>(Context, &SkeletalMeshIn))
+		if (TManagedArray<FVector3f>* Vertices = InCollection.FindAttribute<FVector3f>("Vertex", FGeometryCollection::VerticesGroup))
 		{
-			int32 Index = GetValue<int32>(Context, &BoneIndexIn);
-			if (Index != INDEX_NONE)
+			if (FindInput(&VertexIndicesIn) && FindInput(&VertexIndicesIn)->GetConnection())
 			{
-				if (TManagedArray<FVector3f>* Vertices = InCollection.FindAttribute<FVector3f>("Vertex", FGeometryCollection::VerticesGroup))
+				for (int32 SelectionIndex : GetValue<TArray<int32>>(Context, &VertexIndicesIn))
 				{
-					TArray<int32> BoundVerts;
-					TArray<float> BoundWeights;
-					if (FindInput(&VertexIndicesIn) && FindInput(&VertexIndicesIn)->GetConnection())
+					if (0 <= SelectionIndex && SelectionIndex < Vertices->Num())
 					{
-						for (int32 SelectionIndex : GetValue<TArray<int32>>(Context, &VertexIndicesIn))
-						{
-							if (0 <= SelectionIndex && SelectionIndex < Vertices->Num())
-							{
-								BoundVerts.Add(SelectionIndex);
-							}
-						}
-						BoundWeights.Init(1.0, BoundVerts.Num());
+						BoundVerts.Add(SelectionIndex);
 					}
-					else
+				}
+				if (BoundVerts.Num())
+				{
+					BoundWeights.Init(1.0, BoundVerts.Num());
+					GeometryCollection::Facades::FKinematicBindingFacade Kinematics(InCollection);
+					Kinematics.AddKinematicBinding(Kinematics.SetBoneBindings(INDEX_NONE, BoundVerts, BoundWeights));
+				}
+			}
+			else if (TObjectPtr<USkeletalMesh> SkeletalMesh = GetValue<TObjectPtr<USkeletalMesh>>(Context, &SkeletalMeshIn))
+			{
+				int32 Index = GetValue<int32>(Context, &BoneIndexIn);
+				if (Index != INDEX_NONE)
+				{
+
+					TArray<FTransform> ComponentPose;
+					Dataflow::Animation::GlobalTransforms(SkeletalMesh->GetRefSkeleton(), ComponentPose);
+					if (0 < Index && Index < ComponentPose.Num())
 					{
-						TArray<FTransform> ComponentPose;
-						Dataflow::Animation::GlobalTransforms(SkeletalMesh->GetRefSkeleton(), ComponentPose);
-						if (0 < Index && Index < ComponentPose.Num())
-						{
-							FVector3f BonePosition(ComponentPose[Index].GetTranslation());
+						FVector3f BonePosition(ComponentPose[Index].GetTranslation());
 
-							int NumVertices = Vertices->Num();
-							for (int i = Vertices->Num() - 1; i > 0; i--)
+						int NumVertices = Vertices->Num();
+						for (int i = Vertices->Num() - 1; i > 0; i--)
+						{
+							if ((BonePosition - (*Vertices)[i]).Length() < Radius)
 							{
-								if ((BonePosition - (*Vertices)[i]).Length() < Radius)
-								{
-									BoundVerts.Add(i);
-									BoundWeights.Add(1.0);
-								}
+								BoundVerts.Add(i);
+								BoundWeights.Add(1.0);
 							}
 						}
+
+						if (BoundVerts.Num())
+						{
+							GeometryCollection::Facades::FKinematicBindingFacade Kinematics(InCollection);
+							Kinematics.AddKinematicBinding(Kinematics.SetBoneBindings(Index, BoundVerts, BoundWeights));
+						}
 					}
-					GeometryCollection::Facades::FKinematicBindingFacade Kinematics(InCollection);
-					Kinematics.AddKinematicBinding(Kinematics.SetBoneBindings(Index, BoundVerts, BoundWeights));
 				}
 			}
+
 		}
 		SetValue<DataType>(Context, InCollection, &Collection);
 	}
