ChaosFlesh : Cook/Package Support

diff --git a/Engine/Plugins/Experimental/ChaosFlesh/Source/ChaosFlesh/Private/ChaosFlesh/ChaosFlesh.cpp b/Engine/Plugins/Experimental/ChaosFlesh/Source/ChaosFlesh/Private/ChaosFlesh/ChaosFlesh.cpp
new file mode 100644
--- /dev/null
+++ b/Engine/Plugins/Experimental/ChaosFlesh/Source/ChaosFlesh/Private/ChaosFlesh/ChaosFlesh.cpp
@@ -0,0 +1,5 @@
+// Copyright Epic Games, Inc. All Rights Reserved.
+
+#include "ChaosFlesh/ChaosFlesh.h"
+
+DEFINE_LOG_CATEGORY(LogChaosFlesh);
\ No newline at end of file
diff --git a/Engine/Plugins/Experimental/ChaosFlesh/Source/ChaosFlesh/Private/ChaosFlesh/FleshCollectionUtility.cpp b/Engine/Plugins/Experimental/ChaosFlesh/Source/ChaosFlesh/Private/ChaosFlesh/FleshCollectionUtility.cpp
--- a/Engine/Plugins/Experimental/ChaosFlesh/Source/ChaosFlesh/Private/ChaosFlesh/FleshCollectionUtility.cpp
+++ b/Engine/Plugins/Experimental/ChaosFlesh/Source/ChaosFlesh/Private/ChaosFlesh/FleshCollectionUtility.cpp
@@ -6,6 +6,7 @@
 
 #include "ChaosFlesh/FleshCollectionUtility.h"
 
+#include "ChaosFlesh/ChaosFlesh.h"
 #include "ChaosFlesh/FleshCollection.h"
 #include "ChaosFlesh/GEO.h"
 #include "ChaosFlesh/PB.h"
@@ -16,10 +17,6 @@
 
 #include <fstream>
 
-// logging
-//#include "Logging\LogMacros.h"
-DEFINE_LOG_CATEGORY(LogChaosFlesh);
-
 namespace ChaosFlesh
 {
 	void GetTetFaces(
diff --git a/Engine/Plugins/Experimental/ChaosFlesh/Source/ChaosFlesh/Private/ChaosFlesh/GEO.cpp b/Engine/Plugins/Experimental/ChaosFlesh/Source/ChaosFlesh/Private/ChaosFlesh/GEO.cpp
--- a/Engine/Plugins/Experimental/ChaosFlesh/Source/ChaosFlesh/Private/ChaosFlesh/GEO.cpp
+++ b/Engine/Plugins/Experimental/ChaosFlesh/Source/ChaosFlesh/Private/ChaosFlesh/GEO.cpp
@@ -1112,12 +1112,13 @@ ReadGEO_v19(
 std::unique_ptr<std::istream> 
 unzip(const std::string& filename, std::unique_ptr<IFileHandle>& infile, std::ios::openmode mode=std::ios::in)
 {
+	std::unique_ptr<std::istream> stream;
+#if !(UE_BUILD_SHIPPING && WITH_EDITOR)
 	FPlatformFileManager& FileManager = FPlatformFileManager::Get();
 	IPlatformFile& PlatformFile = FileManager.GetPlatformFile();
 	infile.reset(PlatformFile.OpenRead(*FString(filename.c_str()), false));
 
     FString errmsg;
-    std::unique_ptr<std::istream> stream;
 #ifdef USE_ZLIB
 	if(GZIP_FILE_HEADER header; header.Read(infile.get(), errmsg))
     {
@@ -1132,6 +1133,7 @@ unzip(const std::string& filename, std::unique_ptr<IFileHandle>& infile, std::io
     }
     if(stream)
         stream->imbue(std::locale::classic());
+#endif
     return stream;
 }
 
diff --git a/Engine/Plugins/Experimental/ChaosFlesh/Source/ChaosFlesh/Private/ChaosFlesh/PB.cpp b/Engine/Plugins/Experimental/ChaosFlesh/Source/ChaosFlesh/Private/ChaosFlesh/PB.cpp
--- a/Engine/Plugins/Experimental/ChaosFlesh/Source/ChaosFlesh/Private/ChaosFlesh/PB.cpp
+++ b/Engine/Plugins/Experimental/ChaosFlesh/Source/ChaosFlesh/Private/ChaosFlesh/PB.cpp
@@ -1,5 +1,7 @@
 // Copyright Epic Games, Inc. All Rights Reserved.
 #include "ChaosFlesh/PB.h"
+
+#include "ChaosFlesh/ChaosFlesh.h"
 #include "ChaosFlesh/FleshCollectionUtility.h" // for LogChaosFlesh
 #include "ChaosFlesh/IFileStream.h"
 #include "HAL/PlatformFile.h"
@@ -26,6 +28,7 @@
 std::shared_ptr<std::istream>
 SafeOpenInput(const FString& filename, const bool binary = true)
 {
+#if !(UE_BUILD_SHIPPING && WITH_EDITOR)
 	FPlatformFileManager& FileManager = FPlatformFileManager::Get();
 	IPlatformFile& PlatformFile = FileManager.GetPlatformFile();
 	if (!PlatformFile.FileExists(*filename))
@@ -48,6 +51,7 @@ SafeOpenInput(const FString& filename, const bool binary = true)
 		IFileHandle* infile = PlatformFile.OpenRead(*filename, false);
 		return std::shared_ptr<std::istream>(new IFileStream(infile));
 	}
+#endif
 	return nullptr;
 }
 
diff --git a/Engine/Plugins/Experimental/ChaosFlesh/Source/ChaosFlesh/Public/ChaosFlesh/ChaosFlesh.h b/Engine/Plugins/Experimental/ChaosFlesh/Source/ChaosFlesh/Public/ChaosFlesh/ChaosFlesh.h
new file mode 100644
--- /dev/null
+++ b/Engine/Plugins/Experimental/ChaosFlesh/Source/ChaosFlesh/Public/ChaosFlesh/ChaosFlesh.h
@@ -0,0 +1,9 @@
+// Copyright Epic Games, Inc. All Rights Reserved.
+
+#pragma once
+
+#include "ChaosFlesh/FleshCollection.h"
+#include "ChaosFlesh/TetrahedralCollection.h"
+#include "Logging/LogMacros.h"
+
+CHAOSFLESH_API DECLARE_LOG_CATEGORY_EXTERN(LogChaosFlesh, Log, All);
diff --git a/Engine/Plugins/Experimental/ChaosFlesh/Source/ChaosFlesh/Public/ChaosFlesh/FleshCollectionUtility.h b/Engine/Plugins/Experimental/ChaosFlesh/Source/ChaosFlesh/Public/ChaosFlesh/FleshCollectionUtility.h
--- a/Engine/Plugins/Experimental/ChaosFlesh/Source/ChaosFlesh/Public/ChaosFlesh/FleshCollectionUtility.h
+++ b/Engine/Plugins/Experimental/ChaosFlesh/Source/ChaosFlesh/Public/ChaosFlesh/FleshCollectionUtility.h
@@ -7,8 +7,6 @@
 #include "Templates/UniquePtr.h"
 #include "Misc/AssertionMacros.h"
 
-DECLARE_LOG_CATEGORY_EXTERN(LogChaosFlesh, Verbose, All);
-
 namespace ChaosFlesh 
 {
 	void GetTetFaces(const FIntVector4& Tet, FIntVector3& Face1, FIntVector3& Face2, FIntVector3& Face3, FIntVector3& Face4, const bool invert = false);
diff --git a/Engine/Plugins/Experimental/ChaosFlesh/Source/ChaosFleshNodes/Private/Dataflow/ChaosFleshEngineAssetNodes.cpp b/Engine/Plugins/Experimental/ChaosFlesh/Source/ChaosFleshNodes/Private/Dataflow/ChaosFleshEngineAssetNodes.cpp
--- a/Engine/Plugins/Experimental/ChaosFlesh/Source/ChaosFleshNodes/Private/Dataflow/ChaosFleshEngineAssetNodes.cpp
+++ b/Engine/Plugins/Experimental/ChaosFlesh/Source/ChaosFleshNodes/Private/Dataflow/ChaosFleshEngineAssetNodes.cpp
@@ -3,6 +3,7 @@
 #include "Dataflow/ChaosFleshEngineAssetNodes.h"
 
 #include "Chaos/Math/Poisson.h"
+#include "ChaosFlesh/ChaosFlesh.h"
 #include "ChaosFlesh/FleshAsset.h"
 #include "ChaosFlesh/FleshCollection.h"
 #include "ChaosFlesh/FleshCollectionUtility.h"
@@ -13,8 +14,6 @@
 
 #include UE_INLINE_GENERATED_CPP_BY_NAME(ChaosFleshEngineAssetNodes)
 
-DEFINE_LOG_CATEGORY(LogChaosFlesh);
-
 namespace Dataflow
 {
 	void RegisterChaosFleshEngineAssetNodes()
diff --git a/Engine/Plugins/Experimental/ChaosFlesh/Source/ChaosFleshNodes/Private/Dataflow/ChaosFleshTetrahedralNodes.cpp b/Engine/Plugins/Experimental/ChaosFlesh/Source/ChaosFleshNodes/Private/Dataflow/ChaosFleshTetrahedralNodes.cpp
--- a/Engine/Plugins/Experimental/ChaosFlesh/Source/ChaosFleshNodes/Private/Dataflow/ChaosFleshTetrahedralNodes.cpp
+++ b/Engine/Plugins/Experimental/ChaosFlesh/Source/ChaosFleshNodes/Private/Dataflow/ChaosFleshTetrahedralNodes.cpp
@@ -4,6 +4,7 @@
 #include "Dataflow/ChaosFleshEngineAssetNodes.h"
 
 #include "Chaos/Deformable/Utilities.h"
+#include "ChaosFlesh/ChaosFlesh.h"
 #include "Chaos/Utilities.h"
 #include "Chaos/UniformGrid.h"
 #include "ChaosFlesh/FleshCollection.h"
diff --git a/Engine/Plugins/Experimental/ChaosFlesh/Source/ChaosFleshNodes/Public/Dataflow/ChaosFleshKinematicInitializationNodes.h b/Engine/Plugins/Experimental/ChaosFlesh/Source/ChaosFleshNodes/Public/Dataflow/ChaosFleshKinematicInitializationNodes.h
--- a/Engine/Plugins/Experimental/ChaosFlesh/Source/ChaosFleshNodes/Public/Dataflow/ChaosFleshKinematicInitializationNodes.h
+++ b/Engine/Plugins/Experimental/ChaosFlesh/Source/ChaosFleshNodes/Public/Dataflow/ChaosFleshKinematicInitializationNodes.h
@@ -67,7 +67,7 @@ public:
 	TArray<int32> VertexIndicesIn;
 
 	UPROPERTY(meta = (DataflowInput, DisplayName = "BoneIndex"))
-	int32 BoneIndexIn;
+	int32 BoneIndexIn = INDEX_NONE;
 	
 	FKinematicInitializationDataflowNode(const Dataflow::FNodeParameters& InParam, FGuid InGuid = FGuid::NewGuid())
 		: FDataflowNode(InParam, InGuid)
