ChaosFlesh : Updated plugin name for package builds.

diff --git a/Engine/Plugins/Experimental/ChaosFlesh/Source/ChaosFleshEngine/Private/ChaosFlesh/ChaosFleshEnginePlugin.cpp b/Engine/Plugins/Experimental/ChaosFlesh/Source/ChaosFleshEngine/Private/ChaosFlesh/ChaosFleshEnginePlugin.cpp
--- a/Engine/Plugins/Experimental/ChaosFlesh/Source/ChaosFleshEngine/Private/ChaosFlesh/ChaosFleshEnginePlugin.cpp
+++ b/Engine/Plugins/Experimental/ChaosFlesh/Source/ChaosFleshEngine/Private/ChaosFlesh/ChaosFleshEnginePlugin.cpp
@@ -16,9 +16,6 @@ private:
 	TUniquePtr<Chaos::FFleshCacheAdapter> FleshCacheAdapter;
 };
 
-IMPLEMENT_MODULE( FChaosFleshEnginePlugin, FleshEngine )
-
-
 
 void FChaosFleshEnginePlugin::StartupModule()
 {
@@ -33,4 +30,5 @@ void FChaosFleshEnginePlugin::ShutdownModule()
 }
 
 
+IMPLEMENT_MODULE(FChaosFleshEnginePlugin, ChaosFleshEngine)
 
