ChaosFlesh : Removed incorrect platform defines.

diff --git a/Engine/Plugins/Experimental/ChaosFlesh/Source/ChaosFlesh/Private/ChaosFlesh/GEO.cpp b/Engine/Plugins/Experimental/ChaosFlesh/Source/ChaosFlesh/Private/ChaosFlesh/GEO.cpp
--- a/Engine/Plugins/Experimental/ChaosFlesh/Source/ChaosFlesh/Private/ChaosFlesh/GEO.cpp
+++ b/Engine/Plugins/Experimental/ChaosFlesh/Source/ChaosFlesh/Private/ChaosFlesh/GEO.cpp
@@ -1113,7 +1113,6 @@ std::unique_ptr<std::istream>
 unzip(const std::string& filename, std::unique_ptr<IFileHandle>& infile, std::ios::openmode mode=std::ios::in)
 {
 	std::unique_ptr<std::istream> stream;
-#if !(UE_BUILD_SHIPPING && WITH_EDITOR)
 	FPlatformFileManager& FileManager = FPlatformFileManager::Get();
 	IPlatformFile& PlatformFile = FileManager.GetPlatformFile();
 	infile.reset(PlatformFile.OpenRead(*FString(filename.c_str()), false));
@@ -1133,7 +1132,6 @@ unzip(const std::string& filename, std::unique_ptr<IFileHandle>& infile, std::io
     }
     if(stream)
         stream->imbue(std::locale::classic());
-#endif
     return stream;
 }
 
diff --git a/Engine/Plugins/Experimental/ChaosFlesh/Source/ChaosFlesh/Private/ChaosFlesh/PB.cpp b/Engine/Plugins/Experimental/ChaosFlesh/Source/ChaosFlesh/Private/ChaosFlesh/PB.cpp
--- a/Engine/Plugins/Experimental/ChaosFlesh/Source/ChaosFlesh/Private/ChaosFlesh/PB.cpp
+++ b/Engine/Plugins/Experimental/ChaosFlesh/Source/ChaosFlesh/Private/ChaosFlesh/PB.cpp
@@ -28,7 +28,6 @@
 std::shared_ptr<std::istream>
 SafeOpenInput(const FString& filename, const bool binary = true)
 {
-#if !(UE_BUILD_SHIPPING && WITH_EDITOR)
 	FPlatformFileManager& FileManager = FPlatformFileManager::Get();
 	IPlatformFile& PlatformFile = FileManager.GetPlatformFile();
 	if (!PlatformFile.FileExists(*filename))
@@ -51,7 +50,6 @@ SafeOpenInput(const FString& filename, const bool binary = true)
 		IFileHandle* infile = PlatformFile.OpenRead(*filename, false);
 		return std::shared_ptr<std::istream>(new IFileStream(infile));
 	}
-#endif
 	return nullptr;
 }
 
