ChaosFlesh : Removed GEO PB ZIP from Engine level libraries.
- The geo file reading has been pushed into HeeadlessChaos and stubbed out in ChaosFleshEditor.

diff --git a/Engine/Plugins/Experimental/ChaosFlesh/Source/ChaosFlesh/ChaosFlesh.Build.cs b/Engine/Plugins/Experimental/ChaosFlesh/Source/ChaosFlesh/ChaosFlesh.Build.cs
--- a/Engine/Plugins/Experimental/ChaosFlesh/Source/ChaosFlesh/ChaosFlesh.Build.cs
+++ b/Engine/Plugins/Experimental/ChaosFlesh/Source/ChaosFlesh/ChaosFlesh.Build.cs
@@ -7,39 +7,6 @@ namespace UnrealBuildTool.Rules
 	{
 		public ChaosFlesh(ReadOnlyTargetRules Target) : base(Target)
 		{
-			// We can do specializations for PUBLIC platforms in this file; otherwise it
-			// may not be defined in any given workspace, and references to it will fail.
-			// You can do a string comparison ala:
-			// \code
-			//		Target.Platform.ToString() == "XXX"
-			// \endcode
-			// But that may break our NDA policy which forbids even mentioning NDA'd
-			// platforms in engine code.  A less fragile way to go about it would be to
-			// add a platform specific build file:
-			//
-			// ChaosFlesh_XXX.Build.cs:
-			// \code
-			//    namespace UnrealBuildTool.Rules
-			//    {
-			//        public class ChaosFlesh_XXX : ChaosFlesh
-			//        {
-			//            public ChaosFlesh_XXX(ReadOnlyTargetRules Target) : base(Target) 
-			//            { <Do XXX specific dances here> }
-			//        }
-			//    }
-			// \endcode
-			//
-			// However, adding one of these build files for a public platform will clash with
-			// specializations that the build system makes.  In practice, if it's in
-			// Engine\Source\Programs\UnrealBuildTool\Platform it is public, otherwise it's not.
-			if (Target.Platform == UnrealTargetPlatform.Win64 ||
-				Target.Platform == UnrealTargetPlatform.Linux ||
-				Target.Platform == UnrealTargetPlatform.Mac)
-			{
-				AddEngineThirdPartyPrivateStaticDependencies(Target, "zlib");
-				PrivateDefinitions.Add("USE_ZLIB=1");
-			}
-
 			PublicDependencyModuleNames.AddRange(
 				new string[]
 				{
diff --git a/Engine/Plugins/Experimental/ChaosFlesh/Source/ChaosFlesh/Private/ChaosFlesh/FleshCollectionUtility.cpp b/Engine/Plugins/Experimental/ChaosFlesh/Source/ChaosFlesh/Private/ChaosFlesh/FleshCollectionUtility.cpp
--- a/Engine/Plugins/Experimental/ChaosFlesh/Source/ChaosFlesh/Private/ChaosFlesh/FleshCollectionUtility.cpp
+++ b/Engine/Plugins/Experimental/ChaosFlesh/Source/ChaosFlesh/Private/ChaosFlesh/FleshCollectionUtility.cpp
@@ -8,8 +8,6 @@
 
 #include "ChaosFlesh/ChaosFlesh.h"
 #include "ChaosFlesh/FleshCollection.h"
-#include "ChaosFlesh/GEO.h"
-#include "ChaosFlesh/PB.h"
 
 
 #include "GeometryCollection/TransformCollection.h"
@@ -176,170 +174,4 @@ namespace ChaosFlesh
 				UE_LOG(LogChaosFlesh, Display, TEXT("WARNING: Encountered %d non-manifold tetrahedral mesh faces."), nonManifold);
 		}
 	}
-
-	TUniquePtr<FFleshCollection> ReadTetPBDeformableGeometryCollection(const FString& Filename)
-	{
-		TUniquePtr<FFleshCollection> Collection;
-		UE_LOG(LogChaosFlesh, Display, TEXT("Reading Path %s"), *Filename);
-
-		IO::DeformableGeometryCollectionReader Reader(Filename);
-		if (!Reader.ReadPBScene())
-			return Collection;
-
-		TArray<FVector> *Vertices = nullptr;
-		TArray<FVector> FrameVertices;
-		TArray<FIntVector4> *Elements = nullptr;
-		TArray<FIntVector3> SurfaceElements;
-
-		TArray<IO::DeformableGeometryCollectionReader::TetMesh*> TetMeshes = Reader.GetTetMeshes();
-		for (auto* It : TetMeshes)
-		{
-			Elements = reinterpret_cast<TArray<FIntVector4>*>(&It->Elements);
-			Vertices = &It->Points;
-			UE_LOG(LogChaosFlesh, Display, TEXT("Got %d tets."), Elements->Num());
-			UE_LOG(LogChaosFlesh, Display, TEXT("Got %d rest points."), Vertices->Num());
-			break;
-		}
-		TArray<FIntVector4> ElemTmp;
-		if (!Elements)
-		{
-			Elements = &ElemTmp;
-			UE_LOG(LogChaosFlesh, Display, TEXT("Got %d tets."), Elements->Num());
-		}
-
-		if (!Vertices || Vertices->Num() == 0)
-		{
-			TPair<int32, int32> FrameRange = Reader.ReadFrameRange();
-			if (!Reader.ReadPoints(FrameRange.Key, FrameVertices))
-			{
-				return Collection;
-			}
-			Vertices = &FrameVertices;
-			UE_LOG(LogChaosFlesh, Display, TEXT("Got %d points from frame %d."), Vertices->Num(), FrameRange.Key);
-		}
-
-		GetSurfaceElements(*Elements, SurfaceElements, false);
-		UE_LOG(LogChaosFlesh, Display, TEXT("Got %d tris."), SurfaceElements.Num());
-
-		Collection.Reset(FFleshCollection::NewFleshCollection(*Vertices, SurfaceElements, *Elements));
-		return Collection;
-	}
-
-
-	TUniquePtr<FFleshCollection> ReadGEOFile(const FString& Filename)
-	{
-		TUniquePtr<FFleshCollection> Collection;
-		UE_LOG(LogChaosFlesh, Display, TEXT("Reading Path %s"), *Filename);
-
-		TMap<FString, int32> intVars;
-		TMap<FString, TArray<int32>> intVectorVars;
-		TMap<FString, TArray<float>> floatVectorVars;
-		TMap<FString, TPair<TArray<std::string>, TArray<int32>>> indexedStringVars;
-		if (!ReadGEO(std::string(TCHAR_TO_UTF8(*Filename)), intVars, intVectorVars, floatVectorVars, indexedStringVars))
-		{
-			UE_LOG(LogChaosFlesh, Display, TEXT("Failed to open GEO file: '%s'."), *Filename);
-			return Collection;
-		}
-
-		TArray<FVector> Vertices;
-		TArray<FIntVector4> Elements;
-		TArray<FIntVector3> SurfaceElements;
-
-		auto fvIt = floatVectorVars.Find("position");
-		if (fvIt == nullptr)
-			fvIt = floatVectorVars.Find("P");
-		if (fvIt != nullptr)
-		{
-			const TArray<float>& coords = *fvIt;
-			Vertices.Reserve(coords.Num() / 3);
-			for (size_t i = 0; i < coords.Num(); i += 3)
-			{
-				FVector pt;
-				for (size_t j = 0; j < 3; j++)
-					pt[j] = coords[i + j];
-				Vertices.Add(pt);
-			}
-		}
-		UE_LOG(LogChaosFlesh, Display, TEXT("Got %d points."), Vertices.Num());
-
-		auto ivIt = intVectorVars.Find("pointref.indices");
-		if (ivIt != nullptr)
-		{
-			auto iIt = intVars.Find("Tetrahedron_run:startvertex");
-			int32 startIndex = iIt == nullptr ? 0 : *iIt;
-
-			iIt = intVars.Find("Tetrahedron_run:nprimitives");
-			int32 numTets = iIt == nullptr ? -1 : *iIt;
-
-			UE_LOG(LogChaosFlesh, Display, TEXT("Tet start index: %d num tets: %d"), startIndex, numTets);
-
-			const TArray<int32>& indices = *ivIt;
-			Elements.Reserve(indices.Num() / 4);
-			size_t stopIndex = numTets != -1 ? startIndex + numTets * 4 : indices.Num();
-			for (size_t i = startIndex; i < stopIndex; i += 4)
-			{
-				FIntVector4 tet;
-				for (size_t j = 0; j < 4; j++)
-					tet[j] = indices[i + j];
-				Elements.Add(tet);
-			}
-			UE_LOG(LogChaosFlesh, Display, TEXT("Got %d tets."), Elements.Num());
-		}
-
-		GetSurfaceElements(Elements, SurfaceElements, false);
-		UE_LOG(LogChaosFlesh, Display, TEXT("Got %d tris."), SurfaceElements.Num());
-
-		Collection.Reset(FFleshCollection::NewFleshCollection(Vertices, SurfaceElements, Elements));
-		return Collection;
-	}
-
-	TUniquePtr<FFleshCollection> ReadTetFile(const FString& Filename)
-	{
-		UE_LOG(LogChaosFlesh, Display, TEXT("Reading Path %s"), *Filename);
-		TUniquePtr<FFleshCollection> Collection;
-
-		TArray<FVector> Vertices;
-		TArray<FIntVector4> Elements;  //UE::Math::TIntVector4<int32>
-		TArray<FIntVector3> SurfaceElements;
-
-		TArray<UE::Math::TVector<float>> FloatPos;
-		TArray<Chaos::TVector<int32, 4>> Tets;
-		if (IO::ReadStructure<4>(Filename, FloatPos, Tets))
-		{
-			Vertices.SetNum(FloatPos.Num());
-			for (int32 i = 0; i < FloatPos.Num(); i++)
-				for (int32 j = 0; j < 3; j++)
-					Vertices[i][j] = FloatPos[i][j];
-
-			Elements.SetNum(Tets.Num());
-			for (int32 i = 0; i < Tets.Num(); i++)
-				for (int32 j = 0; j < 4; j++)
-					Elements[i][j] = Tets[i][j];
-
-			UE_LOG(LogChaosFlesh, Display, TEXT("Got %d tets."), Elements.Num());
-		}
-
-		GetSurfaceElements(Elements, SurfaceElements, false);
-		UE_LOG(LogChaosFlesh, Display, TEXT("Got %d tris."), SurfaceElements.Num());
-
-		Collection.Reset(FFleshCollection::NewFleshCollection(Vertices, SurfaceElements, Elements));
-		return Collection;
-	}
-
-	TUniquePtr<FFleshCollection> ImportTetFromFile(const FString& Filename)
-	{
-		if (FPaths::FileExists(Filename))
-		{
-			if (Filename.EndsWith(".tet") || Filename.EndsWith(".tet.gz"))
-				return ReadTetFile(Filename);
-			else if (Filename.EndsWith(".geo"))
-				return ReadGEOFile(Filename);
-			UE_LOG(LogChaosFlesh, Warning, TEXT("Unsupported file type: '%s'."), *Filename);
-		}
-		else
-		{
-			UE_LOG(LogChaosFlesh, Warning, TEXT("Unknown file path: '%s'."), *Filename);
-		}
-		return nullptr;
-	}
 }
diff --git a/Engine/Plugins/Experimental/ChaosFlesh/Source/ChaosFlesh/Private/ChaosFlesh/GEO.cpp b/Engine/Plugins/Experimental/ChaosFlesh/Source/ChaosFlesh/Private/ChaosFlesh/GEO.cpp
deleted file mode 100644
diff --git a/Engine/Plugins/Experimental/ChaosFlesh/Source/ChaosFlesh/Private/ChaosFlesh/GEO.h b/Engine/Plugins/Experimental/ChaosFlesh/Source/ChaosFlesh/Private/ChaosFlesh/GEO.h
deleted file mode 100644
diff --git a/Engine/Plugins/Experimental/ChaosFlesh/Source/ChaosFlesh/Private/ChaosFlesh/IFileStream.cpp b/Engine/Plugins/Experimental/ChaosFlesh/Source/ChaosFlesh/Private/ChaosFlesh/IFileStream.cpp
deleted file mode 100644
diff --git a/Engine/Plugins/Experimental/ChaosFlesh/Source/ChaosFlesh/Private/ChaosFlesh/IFileStream.h b/Engine/Plugins/Experimental/ChaosFlesh/Source/ChaosFlesh/Private/ChaosFlesh/IFileStream.h
deleted file mode 100644
diff --git a/Engine/Plugins/Experimental/ChaosFlesh/Source/ChaosFlesh/Private/ChaosFlesh/PB.cpp b/Engine/Plugins/Experimental/ChaosFlesh/Source/ChaosFlesh/Private/ChaosFlesh/PB.cpp
deleted file mode 100644
diff --git a/Engine/Plugins/Experimental/ChaosFlesh/Source/ChaosFlesh/Private/ChaosFlesh/PB.h b/Engine/Plugins/Experimental/ChaosFlesh/Source/ChaosFlesh/Private/ChaosFlesh/PB.h
deleted file mode 100644
diff --git a/Engine/Plugins/Experimental/ChaosFlesh/Source/ChaosFlesh/Private/ChaosFlesh/ZIP.h b/Engine/Plugins/Experimental/ChaosFlesh/Source/ChaosFlesh/Private/ChaosFlesh/ZIP.h
deleted file mode 100644
diff --git a/Engine/Plugins/Experimental/ChaosFlesh/Source/ChaosFlesh/Public/ChaosFlesh/FleshCollectionUtility.h b/Engine/Plugins/Experimental/ChaosFlesh/Source/ChaosFlesh/Public/ChaosFlesh/FleshCollectionUtility.h
--- a/Engine/Plugins/Experimental/ChaosFlesh/Source/ChaosFlesh/Public/ChaosFlesh/FleshCollectionUtility.h
+++ b/Engine/Plugins/Experimental/ChaosFlesh/Source/ChaosFlesh/Public/ChaosFlesh/FleshCollectionUtility.h
@@ -11,11 +11,4 @@ namespace ChaosFlesh
 {
 	void GetTetFaces(const FIntVector4& Tet, FIntVector3& Face1, FIntVector3& Face2, FIntVector3& Face3, FIntVector3& Face4, const bool invert = false);
 	void GetSurfaceElements(const TArray<FIntVector4>& Tets, TArray<FIntVector3>& SurfaceElements, const bool KeepInteriorFaces, const bool InvertFaces = false);
-
-	TUniquePtr<FFleshCollection> ReadGEOFile(const FString& Filename);
-	TUniquePtr<FFleshCollection> ReadTetFile(const FString& Filename); // uncompressed tet
-	TUniquePtr<FFleshCollection> ReadTetPBDeformableGeometryCollection(const FString& Filename);
-
-	//! Supports reading .tet and .geo (compatible with version 19).
-	TUniquePtr<FFleshCollection> CHAOSFLESH_API ImportTetFromFile(const FString& Filename);
 }
diff --git a/Engine/Plugins/Experimental/ChaosFlesh/Source/ChaosFleshEditor/Private/ChaosFlesh/Cmd/ChaosFleshCommands.cpp b/Engine/Plugins/Experimental/ChaosFlesh/Source/ChaosFleshEditor/Private/ChaosFlesh/Cmd/ChaosFleshCommands.cpp
--- a/Engine/Plugins/Experimental/ChaosFlesh/Source/ChaosFleshEditor/Private/ChaosFlesh/Cmd/ChaosFleshCommands.cpp
+++ b/Engine/Plugins/Experimental/ChaosFlesh/Source/ChaosFleshEditor/Private/ChaosFlesh/Cmd/ChaosFleshCommands.cpp
@@ -3,6 +3,7 @@
 #include "ChaosFlesh/Cmd/ChaosFleshCommands.h"
 
 #include "ChaosFlesh/Asset/FleshAssetFactory.h"
+#include "ChaosFlesh/Cmd/FleshAssetConversion.h"
 #include "ChaosFlesh/FleshAsset.h"
 #include "ChaosFlesh/FleshCollection.h"
 #include "ChaosFlesh/FleshCollectionUtility.h"
@@ -35,7 +36,7 @@ void FChaosFleshCommands::ImportFile(const TArray<FString>& Args, UWorld* World)
 				if (FFleshCollection* Collection = EditObject.GetFleshCollection())
 				{
 					UE_LOG(UChaosFleshCommandsLogging, Log, TEXT("FChaosFleshCommands::ImportFile"));
-					if (TUniquePtr<FFleshCollection> InCollection = ChaosFlesh::ImportTetFromFile(Args[0]))
+					if (TUniquePtr<FFleshCollection> InCollection = FFleshAssetConversion::ImportTetFromFile(Args[0]))
 					{
 						Collection->CopyMatchingAttributesFrom(*InCollection);
 					}
diff --git a/Engine/Plugins/Experimental/ChaosFlesh/Source/ChaosFleshEditor/Private/ChaosFlesh/Cmd/FleshAssetConversion.cpp b/Engine/Plugins/Experimental/ChaosFlesh/Source/ChaosFleshEditor/Private/ChaosFlesh/Cmd/FleshAssetConversion.cpp
--- a/Engine/Plugins/Experimental/ChaosFlesh/Source/ChaosFleshEditor/Private/ChaosFlesh/Cmd/FleshAssetConversion.cpp
+++ b/Engine/Plugins/Experimental/ChaosFlesh/Source/ChaosFleshEditor/Private/ChaosFlesh/Cmd/FleshAssetConversion.cpp
@@ -4,10 +4,178 @@
 
 #include "Logging/LogMacros.h"
 
+//#include "ChaosFlesh/PB.h"
+//#include "ChaosFlesh/GEO.h"
 
 DEFINE_LOG_CATEGORY_STATIC(UFleshAssetConversionLogging, Log, All);
 
+/*
+TUniquePtr<FFleshCollection> ReadTetPBDeformableGeometryCollection(const FString& Filename)
+	{
+		TUniquePtr<FFleshCollection> Collection;
+		UE_LOG(LogChaosFlesh, Display, TEXT("Reading Path %s"), *Filename);
 
-void FFleshAssetConversion::ImportFile(FString & FileName)
+		IO::DeformableGeometryCollectionReader Reader(Filename);
+		if (!Reader.ReadPBScene())
+			return Collection;
+
+		TArray<FVector> *Vertices = nullptr;
+		TArray<FVector> FrameVertices;
+		TArray<FIntVector4> *Elements = nullptr;
+		TArray<FIntVector3> SurfaceElements;
+
+		TArray<IO::DeformableGeometryCollectionReader::TetMesh*> TetMeshes = Reader.GetTetMeshes();
+		for (auto* It : TetMeshes)
+		{
+			Elements = reinterpret_cast<TArray<FIntVector4>*>(&It->Elements);
+			Vertices = &It->Points;
+			UE_LOG(LogChaosFlesh, Display, TEXT("Got %d tets."), Elements->Num());
+			UE_LOG(LogChaosFlesh, Display, TEXT("Got %d rest points."), Vertices->Num());
+			break;
+		}
+		TArray<FIntVector4> ElemTmp;
+		if (!Elements)
+		{
+			Elements = &ElemTmp;
+			UE_LOG(LogChaosFlesh, Display, TEXT("Got %d tets."), Elements->Num());
+		}
+
+		if (!Vertices || Vertices->Num() == 0)
+		{
+			TPair<int32, int32> FrameRange = Reader.ReadFrameRange();
+			if (!Reader.ReadPoints(FrameRange.Key, FrameVertices))
+			{
+				return Collection;
+			}
+			Vertices = &FrameVertices;
+			UE_LOG(LogChaosFlesh, Display, TEXT("Got %d points from frame %d."), Vertices->Num(), FrameRange.Key);
+		}
+
+		GetSurfaceElements(*Elements, SurfaceElements, false);
+		UE_LOG(LogChaosFlesh, Display, TEXT("Got %d tris."), SurfaceElements.Num());
+
+		Collection.Reset(FFleshCollection::NewFleshCollection(*Vertices, SurfaceElements, *Elements));
+		return Collection;
+	}
+
+	TUniquePtr<FFleshCollection> ReadGEOFile(const FString& Filename)
+	{
+		TUniquePtr<FFleshCollection> Collection;
+		UE_LOG(LogChaosFlesh, Display, TEXT("Reading Path %s"), *Filename);
+
+		TMap<FString, int32> intVars;
+		TMap<FString, TArray<int32>> intVectorVars;
+		TMap<FString, TArray<float>> floatVectorVars;
+		TMap<FString, TPair<TArray<std::string>, TArray<int32>>> indexedStringVars;
+		if (!ReadGEO(std::string(TCHAR_TO_UTF8(*Filename)), intVars, intVectorVars, floatVectorVars, indexedStringVars))
+		{
+			UE_LOG(LogChaosFlesh, Display, TEXT("Failed to open GEO file: '%s'."), *Filename);
+			return Collection;
+		}
+
+		TArray<FVector> Vertices;
+		TArray<FIntVector4> Elements;
+		TArray<FIntVector3> SurfaceElements;
+
+		auto fvIt = floatVectorVars.Find("position");
+		if (fvIt == nullptr)
+			fvIt = floatVectorVars.Find("P");
+		if (fvIt != nullptr)
+		{
+			const TArray<float>& coords = *fvIt;
+			Vertices.Reserve(coords.Num() / 3);
+			for (size_t i = 0; i < coords.Num(); i += 3)
+			{
+				FVector pt;
+				for (size_t j = 0; j < 3; j++)
+					pt[j] = coords[i + j];
+				Vertices.Add(pt);
+			}
+		}
+		UE_LOG(LogChaosFlesh, Display, TEXT("Got %d points."), Vertices.Num());
+
+		auto ivIt = intVectorVars.Find("pointref.indices");
+		if (ivIt != nullptr)
+		{
+			auto iIt = intVars.Find("Tetrahedron_run:startvertex");
+			int32 startIndex = iIt == nullptr ? 0 : *iIt;
+
+			iIt = intVars.Find("Tetrahedron_run:nprimitives");
+			int32 numTets = iIt == nullptr ? -1 : *iIt;
+
+			UE_LOG(LogChaosFlesh, Display, TEXT("Tet start index: %d num tets: %d"), startIndex, numTets);
+
+			const TArray<int32>& indices = *ivIt;
+			Elements.Reserve(indices.Num() / 4);
+			size_t stopIndex = numTets != -1 ? startIndex + numTets * 4 : indices.Num();
+			for (size_t i = startIndex; i < stopIndex; i += 4)
+			{
+				FIntVector4 tet;
+				for (size_t j = 0; j < 4; j++)
+					tet[j] = indices[i + j];
+				Elements.Add(tet);
+			}
+			UE_LOG(LogChaosFlesh, Display, TEXT("Got %d tets."), Elements.Num());
+		}
+
+		GetSurfaceElements(Elements, SurfaceElements, false);
+		UE_LOG(LogChaosFlesh, Display, TEXT("Got %d tris."), SurfaceElements.Num());
+
+		Collection.Reset(FFleshCollection::NewFleshCollection(Vertices, SurfaceElements, Elements));
+		return Collection;
+	}
+
+	TUniquePtr<FFleshCollection> ReadTetFile(const FString& Filename)
+	{
+		UE_LOG(LogChaosFlesh, Display, TEXT("Reading Path %s"), *Filename);
+		TUniquePtr<FFleshCollection> Collection;
+
+		TArray<FVector> Vertices;
+		TArray<FIntVector4> Elements;  //UE::Math::TIntVector4<int32>
+		TArray<FIntVector3> SurfaceElements;
+
+		TArray<UE::Math::TVector<float>> FloatPos;
+		TArray<Chaos::TVector<int32, 4>> Tets;
+		if (IO::ReadStructure<4>(Filename, FloatPos, Tets))
+		{
+			Vertices.SetNum(FloatPos.Num());
+			for (int32 i = 0; i < FloatPos.Num(); i++)
+				for (int32 j = 0; j < 3; j++)
+					Vertices[i][j] = FloatPos[i][j];
+
+			Elements.SetNum(Tets.Num());
+			for (int32 i = 0; i < Tets.Num(); i++)
+				for (int32 j = 0; j < 4; j++)
+					Elements[i][j] = Tets[i][j];
+
+			UE_LOG(LogChaosFlesh, Display, TEXT("Got %d tets."), Elements.Num());
+		}
+
+		GetSurfaceElements(Elements, SurfaceElements, false);
+		UE_LOG(LogChaosFlesh, Display, TEXT("Got %d tris."), SurfaceElements.Num());
+
+		Collection.Reset(FFleshCollection::NewFleshCollection(Vertices, SurfaceElements, Elements));
+		return Collection;
+	}
+*/
+
+TUniquePtr<FFleshCollection> FFleshAssetConversion::ImportTetFromFile(const FString& Filename)
 {
+	ensure(false);
+	/*
+	if (FPaths::FileExists(Filename))
+	{
+		if (Filename.EndsWith(".tet") || Filename.EndsWith(".tet.gz"))
+			return ReadTetFile(Filename);
+		else if (Filename.EndsWith(".geo"))
+			return ReadGEOFile(Filename);
+		UE_LOG(LogChaosFlesh, Warning, TEXT("Unsupported file type: '%s'."), *Filename);
+	}
+	else
+	{
+		UE_LOG(LogChaosFlesh, Warning, TEXT("Unknown file path: '%s'."), *Filename);
+	}
+	*/
+	return nullptr;
 }
+
diff --git a/Engine/Plugins/Experimental/ChaosFlesh/Source/ChaosFleshEditor/Public/ChaosFlesh/Cmd/FleshAssetConversion.h b/Engine/Plugins/Experimental/ChaosFlesh/Source/ChaosFleshEditor/Public/ChaosFlesh/Cmd/FleshAssetConversion.h
--- a/Engine/Plugins/Experimental/ChaosFlesh/Source/ChaosFleshEditor/Public/ChaosFlesh/Cmd/FleshAssetConversion.h
+++ b/Engine/Plugins/Experimental/ChaosFlesh/Source/ChaosFleshEditor/Public/ChaosFlesh/Cmd/FleshAssetConversion.h
@@ -2,6 +2,8 @@
 
 #pragma once
 #include "Engine/World.h"
+#include "ChaosFlesh/FleshCollection.h"
+#include "Templates/UniquePtr.h"
 
 
 /**
@@ -12,7 +14,15 @@ class CHAOSFLESHEDITOR_API FFleshAssetConversion
 public:
 
 	/**
+	* Supports reading .tet and .geo (compatible with version 19).
+	* Currently disabled.
 	*/
-	static void ImportFile(FString& FileName);
+	static TUniquePtr<FFleshCollection> ImportTetFromFile(const FString& Filename);
+
+
+	//TUniquePtr<FFleshCollection> ReadGEOFile(const FString& Filename);
+	//TUniquePtr<FFleshCollection> ReadTetFile(const FString& Filename); // uncompressed tet
+	//TUniquePtr<FFleshCollection> ReadTetPBDeformableGeometryCollection(const FString& Filename)
+
 
 };
\ No newline at end of file
diff --git a/Engine/Plugins/Experimental/ChaosFlesh/Source/ChaosFleshNodes/Private/Dataflow/ChaosFleshEngineAssetNodes.cpp b/Engine/Plugins/Experimental/ChaosFlesh/Source/ChaosFleshNodes/Private/Dataflow/ChaosFleshEngineAssetNodes.cpp
--- a/Engine/Plugins/Experimental/ChaosFlesh/Source/ChaosFleshNodes/Private/Dataflow/ChaosFleshEngineAssetNodes.cpp
+++ b/Engine/Plugins/Experimental/ChaosFlesh/Source/ChaosFleshNodes/Private/Dataflow/ChaosFleshEngineAssetNodes.cpp
@@ -20,7 +20,6 @@ namespace Dataflow
 	{
 		DATAFLOW_NODE_REGISTER_CREATION_FACTORY(FGetFleshAssetDataflowNode);
 		DATAFLOW_NODE_REGISTER_CREATION_FACTORY(FFleshAssetTerminalDataflowNode);
-		DATAFLOW_NODE_REGISTER_CREATION_FACTORY(FImportFleshDataflowNode);
 		DATAFLOW_NODE_REGISTER_CREATION_FACTORY(FComputeFleshMassNode);
 		DATAFLOW_NODE_REGISTER_CREATION_FACTORY(FComputeFiberFieldNode);
 	}
@@ -58,24 +57,6 @@ void FFleshAssetTerminalDataflowNode::Evaluate(Dataflow::FContext& Context, cons
 	SetValue<FManagedArrayCollection>(Context, InCollection, &Collection);
 }
 
-void FImportFleshDataflowNode::Evaluate(Dataflow::FContext& Context, const FDataflowOutput* Out) const
-{
-	if (Out->IsA<FManagedArrayCollection>(&Collection))
-	{
-		FString TempFileName = Filename.FilePath;
-		TempFileName = TempFileName.Replace(TEXT("\\"), TEXT("/"));
-		if (TUniquePtr<FFleshCollection> FleshCollection = ChaosFlesh::ImportTetFromFile(TempFileName))
-		{
-			SetValue<FManagedArrayCollection>(Context, *FleshCollection.Release(), &Collection);
-		}
-		else
-		{
-			SetValue<FManagedArrayCollection>(Context, Collection, &Collection);
-		}
-	}
-}
-
-
 template <class T> using MType = FManagedArrayCollection::TManagedType<T>;
 
 void FComputeFleshMassNode::Evaluate(Dataflow::FContext& Context, const FDataflowOutput* Out) const
diff --git a/Engine/Plugins/Experimental/ChaosFlesh/Source/ChaosFleshNodes/Private/Dataflow/ChaosFleshKinematicInitializationNodes.cpp b/Engine/Plugins/Experimental/ChaosFlesh/Source/ChaosFleshNodes/Private/Dataflow/ChaosFleshKinematicInitializationNodes.cpp
--- a/Engine/Plugins/Experimental/ChaosFlesh/Source/ChaosFleshNodes/Private/Dataflow/ChaosFleshKinematicInitializationNodes.cpp
+++ b/Engine/Plugins/Experimental/ChaosFlesh/Source/ChaosFleshNodes/Private/Dataflow/ChaosFleshKinematicInitializationNodes.cpp
@@ -23,7 +23,6 @@ namespace Dataflow
 		DATAFLOW_NODE_REGISTER_CREATION_FACTORY(FKinematicInitializationDataflowNode);
 		DATAFLOW_NODE_REGISTER_CREATION_FACTORY(FKinematicTetrahedralBindingsDataflowNode);
 		DATAFLOW_NODE_REGISTER_CREATION_FACTORY(FSetVerticesKinematicDataflowNode);
-		DATAFLOW_NODE_REGISTER_CREATION_FACTORY(FBinVerticesDataflowNode);
 	}
 }
 
@@ -244,37 +243,4 @@ void FSetVerticesKinematicDataflowNode::Evaluate(Dataflow::FContext& Context, co
 		}
 		SetValue<DataType>(Context, InCollection, &Collection);
 	}
-}
-
-
-void FBinVerticesDataflowNode::Evaluate(Dataflow::FContext& Context, const FDataflowOutput* Out) const
-{
-	if (Out->IsA<DataType>(&VertexIndicesOut))
-	{	
-		TArray<int32> VertexArray;
-		FManagedArrayCollection InCollection = GetValue<FManagedArrayCollection>(Context, &Collection);
-		FString TempFileName = Filename;
-		TempFileName = TempFileName.Replace(TEXT("\\"), TEXT("/"));
-		if (TUniquePtr<FFleshCollection> FleshCollection = ChaosFlesh::ImportTetFromFile(TempFileName))
-		{	
-			TManagedArray<FVector3f>* VerticesBin = FleshCollection->FindAttribute<FVector3f>("Vertex", FGeometryCollection::VerticesGroup);
-			TManagedArray<FVector3f>* VerticesCollection = InCollection.FindAttribute<FVector3f>("Vertex", FGeometryCollection::VerticesGroup);
-			if (VerticesBin && VerticesCollection)
-			{
-				for (int32 i = 0; i < VerticesBin->Num(); i++)
-				{
-					for (int32 j = 0; j < VerticesCollection->Num(); ++j)
-					{
-						if (((*VerticesBin)[i] - (*VerticesCollection)[j]).Length() < Tolerance)
-						{
-							VertexArray.Emplace(j);
-							break;
-						}
-					}
-					ensureMsgf(VertexArray.Num() == i + 1, TEXT("Some vertex is not binned or binned repeatedly at current tolerance"));
-				}
-			}
-		}
-		Out->SetValue<DataType>(VertexArray, Context);
-	}
-}
+}
\ No newline at end of file
diff --git a/Engine/Plugins/Experimental/ChaosFlesh/Source/ChaosFleshNodes/Public/Dataflow/ChaosFleshEngineAssetNodes.h b/Engine/Plugins/Experimental/ChaosFlesh/Source/ChaosFleshNodes/Public/Dataflow/ChaosFleshEngineAssetNodes.h
--- a/Engine/Plugins/Experimental/ChaosFlesh/Source/ChaosFleshNodes/Public/Dataflow/ChaosFleshEngineAssetNodes.h
+++ b/Engine/Plugins/Experimental/ChaosFlesh/Source/ChaosFleshNodes/Public/Dataflow/ChaosFleshEngineAssetNodes.h
@@ -59,32 +59,6 @@ public:
 };
 
 
-
-USTRUCT(meta = (DataflowFlesh))
-struct FImportFleshDataflowNode : public FDataflowNode
-{
-	GENERATED_USTRUCT_BODY()
-	DATAFLOW_NODE_DEFINE_INTERNAL(FImportFleshDataflowNode, "ImportFlesh", "Flesh", "")
-	DATAFLOW_NODE_RENDER_TYPE(FGeometryCollection::StaticType(), "Collection")
-
-public:
-
-	UPROPERTY(EditAnywhere, Category = "Dataflow")
-	FFilePath Filename;
-
-	UPROPERTY(meta = (DataflowOutput, DisplayName = "Collection"))
-	FManagedArrayCollection Collection;
-	
-	FImportFleshDataflowNode(const Dataflow::FNodeParameters& InParam, FGuid InGuid = FGuid::NewGuid())
-		: FDataflowNode(InParam, InGuid)
-	{
-		RegisterOutputConnection(&Collection);
-	}
-
-	virtual void Evaluate(Dataflow::FContext& Context, const FDataflowOutput* Out) const override;
-};
-
-
 USTRUCT(meta = (DataflowFlesh))
 struct FComputeFleshMassNode : public FDataflowNode
 {
diff --git a/Engine/Plugins/Experimental/ChaosFlesh/Source/ChaosFleshNodes/Public/Dataflow/ChaosFleshKinematicInitializationNodes.h b/Engine/Plugins/Experimental/ChaosFlesh/Source/ChaosFleshNodes/Public/Dataflow/ChaosFleshKinematicInitializationNodes.h
--- a/Engine/Plugins/Experimental/ChaosFlesh/Source/ChaosFleshNodes/Public/Dataflow/ChaosFleshKinematicInitializationNodes.h
+++ b/Engine/Plugins/Experimental/ChaosFlesh/Source/ChaosFleshNodes/Public/Dataflow/ChaosFleshKinematicInitializationNodes.h
@@ -110,34 +110,6 @@ public:
 	virtual void Evaluate(Dataflow::FContext& Context, const FDataflowOutput* Out) const override;
 };
 
-USTRUCT(meta = (DataflowFlesh))
-struct FBinVerticesDataflowNode : public FDataflowNode
-{
-	GENERATED_USTRUCT_BODY()
-	DATAFLOW_NODE_DEFINE_INTERNAL(FBinVerticesDataflowNode, "BinVertices", "Flesh", "")
-	DATAFLOW_NODE_RENDER_TYPE(FGeometryCollection::StaticType(), "Collection")
-
-public:
-	typedef TArray<int32> DataType;
-	UPROPERTY(EditAnywhere, Category = "Dataflow")
-		FString Filename = FString("D:/UE5/Main/QAGame/Import/example.geo");
-	UPROPERTY(EditAnywhere, Category = "Dataflow")
-		float Tolerance = float(1e-6);
-	UPROPERTY(meta = (DataflowInput, DisplayName = "Collection"))
-		FManagedArrayCollection Collection;
-	UPROPERTY(meta = (DataflowOutput, DisplayName = "SelectionSet"))
-		TArray<int32> VertexIndicesOut;
-
-	FBinVerticesDataflowNode(const Dataflow::FNodeParameters& InParam, FGuid InGuid = FGuid::NewGuid())
-		: FDataflowNode(InParam, InGuid)
-	{
-		RegisterInputConnection(&Collection);
-		RegisterOutputConnection(&VertexIndicesOut);
-	}
-
-	virtual void Evaluate(Dataflow::FContext& Context, const FDataflowOutput* Out) const override;
-};
-
 namespace Dataflow
 {
 	void RegisterChaosFleshKinematicInitializationNodes();
