ChaosFlesh : Fixed incorect node default.

diff --git a/Engine/Plugins/Experimental/ChaosFlesh/Source/ChaosFleshNodes/Private/Dataflow/ChaosFleshTetrahedralNodes.cpp b/Engine/Plugins/Experimental/ChaosFlesh/Source/ChaosFleshNodes/Private/Dataflow/ChaosFleshTetrahedralNodes.cpp
--- a/Engine/Plugins/Experimental/ChaosFlesh/Source/ChaosFleshNodes/Private/Dataflow/ChaosFleshTetrahedralNodes.cpp
+++ b/Engine/Plugins/Experimental/ChaosFlesh/Source/ChaosFleshNodes/Private/Dataflow/ChaosFleshTetrahedralNodes.cpp
@@ -136,15 +136,14 @@ void FGenerateTetrahedralCollectionDataflowNodes::EvaluateIsoStuffing(
 	TObjectPtr<const UStaticMesh>& InStaticMesh) const
 {
 #if WITH_EDITORONLY_DATA
-	const TObjectPtr<const UStaticMesh>& StaticMeshValue = GetValue<TObjectPtr<const UStaticMesh>>(Context, &InStaticMesh);
-	if (StaticMeshValue && NumCells > 0 && (-.5 <= OffsetPercent && OffsetPercent <= 0.5))
+	if (InStaticMesh && NumCells > 0 && (-.5 <= OffsetPercent && OffsetPercent <= 0.5))
 	{
 		// make a mesh description for UE::Geometry tools
 		UE::Geometry::FDynamicMesh3 DynamicMesh;
 		FMeshDescriptionToDynamicMesh GetSourceMesh;
-		bool bUsingHiResSource = StaticMeshValue->IsHiResMeshDescriptionValid();
+		bool bUsingHiResSource = InStaticMesh->IsHiResMeshDescriptionValid();
 		const FMeshDescription* UseSourceMeshDescription =
-			(bUsingHiResSource) ? StaticMeshValue->GetHiResMeshDescription() : StaticMeshValue->GetMeshDescription(0);
+			(bUsingHiResSource) ? InStaticMesh->GetHiResMeshDescription() : InStaticMesh->GetMeshDescription(0);
 		GetSourceMesh.Convert(UseSourceMeshDescription, DynamicMesh);
 
 		// Tet mesh generation
diff --git a/Engine/Source/Runtime/Experimental/Dataflow/Core/Private/Dataflow/DataflowInputOutput.cpp b/Engine/Source/Runtime/Experimental/Dataflow/Core/Private/Dataflow/DataflowInputOutput.cpp
--- a/Engine/Source/Runtime/Experimental/Dataflow/Core/Private/Dataflow/DataflowInputOutput.cpp
+++ b/Engine/Source/Runtime/Experimental/Dataflow/Core/Private/Dataflow/DataflowInputOutput.cpp
@@ -23,7 +23,6 @@ FDataflowInput::FDataflowInput(const Dataflow::FInputParameters& Param, FGuid In
 
 bool FDataflowInput::AddConnection(FDataflowConnection* InOutput)
 {
-	ensure(Connection == nullptr);
 	if (ensure(InOutput->GetType() == this->GetType()))
 	{
 		Connection = (FDataflowOutput*)InOutput;
