ChaosFlesh : Simulation and asset updates.
- Added FStringValuePair for attribute overrides within dataflow nodes.
- Added Skeleton attribute to asset.
- Tetrahedral Alignment with Skeleton.
-  Debug draw for solver state.
- User defined tick locations within the physics ticks
- Threaded and non threaded evaluation modes.
- Update to Kinematic Dataflow Node
- Allow the initialization from the Kinematic node to bind to the skeleton.

diff --git a/Engine/Plugins/Experimental/ChaosFlesh/Source/ChaosFleshEngine/Private/ChaosFlesh/ChaosDeformableSolverComponent.cpp b/Engine/Plugins/Experimental/ChaosFlesh/Source/ChaosFleshEngine/Private/ChaosFlesh/ChaosDeformableSolverComponent.cpp
--- a/Engine/Plugins/Experimental/ChaosFlesh/Source/ChaosFleshEngine/Private/ChaosFlesh/ChaosDeformableSolverComponent.cpp
+++ b/Engine/Plugins/Experimental/ChaosFlesh/Source/ChaosFleshEngine/Private/ChaosFlesh/ChaosDeformableSolverComponent.cpp
@@ -20,19 +20,52 @@ UDeformableSolverComponent::UDeformableSolverComponent(const FObjectInitializer&
 	PrimaryComponentTick.bCanEverTick = true;
 	bTickInEditor = false;
 
-	PrimaryComponentTick.bCanEverTick = true;
-	PrimaryComponentTick.bTickEvenWhenPaused = false;
-	PrimaryComponentTick.TickGroup = TG_PrePhysics;
-
-	DeformableEndTickFunction.TickGroup = TG_PostPhysics;
-	DeformableEndTickFunction.bCanEverTick = true;
-	DeformableEndTickFunction.bStartWithTickEnabled = true;
+	UpdateTickGroup();
 }
 
 UDeformableSolverComponent::~UDeformableSolverComponent()
 {
 }
 
+void UDeformableSolverComponent::UpdateTickGroup()
+{
+	//
+	// OR THIS SolverData->PreSolveHandle  = Solver->AddPreAdvanceCallback(FSolverPreAdvance::FDelegate::CreateUObject(this, &AChaosCacheManager::HandlePreSolve, Solver));
+	// see : CacheManagerActor.cpp::348
+	
+	//
+	if (ExecutionModel == EDeformableExecutionModel::Chaos_Deformable_PrePhysics)
+	{
+		PrimaryComponentTick.TickGroup = TG_PrePhysics;
+		DeformableEndTickFunction.TickGroup = TG_PrePhysics;
+	}
+	else if (ExecutionModel == EDeformableExecutionModel::Chaos_Deformable_PostPhysics)
+	{
+		PrimaryComponentTick.TickGroup = TG_PostPhysics;
+		DeformableEndTickFunction.TickGroup = TG_LastDemotable;
+	}
+	else //EDeformableExecutionModel::Chaos_Deformable_DuringPhysics
+	{
+		PrimaryComponentTick.TickGroup = TG_PrePhysics;
+		DeformableEndTickFunction.TickGroup = TG_PostPhysics;
+	}
+
+	PrimaryComponentTick.bCanEverTick = true;
+	PrimaryComponentTick.bTickEvenWhenPaused = false;
+
+	if (bDoThreadedAdvance)
+	{
+		DeformableEndTickFunction.bCanEverTick = true;
+		DeformableEndTickFunction.bStartWithTickEnabled = true;
+	}
+	else
+	{
+		DeformableEndTickFunction.bCanEverTick = false;
+		DeformableEndTickFunction.bStartWithTickEnabled = false;
+	}
+}
+
+
 UDeformableSolverComponent::FDeformableSolver::FGameThreadAccess
 UDeformableSolverComponent::GameThreadAccess()
 {
@@ -122,33 +155,35 @@ void UDeformableSolverComponent::UpdateDeformableEndTickState(bool bRegister)
 {
 	UE_LOG(LogDeformableSolverComponentInternal, Verbose, TEXT("UDeformableSolverComponent::RegiUpdateDeformableEndTickStatesterEndPhysicsTick"));
 	bRegister &= PrimaryComponentTick.IsTickFunctionRegistered();
-
-	if (bRegister != DeformableEndTickFunction.IsTickFunctionRegistered())
+	if (bDoThreadedAdvance)
 	{
-		if (bRegister)
+		if (bRegister != DeformableEndTickFunction.IsTickFunctionRegistered())
 		{
-			UWorld* World = GetWorld();
-			if (World->EndPhysicsTickFunction.IsTickFunctionRegistered() && SetupActorComponentTickFunction(&DeformableEndTickFunction))
+			if (bRegister)
 			{
-				DeformableEndTickFunction.DeformableSolverComponent = this;
-				// Make sure our EndPhysicsTick gets called after physics simulation is finished
-				if (World != nullptr)
+				UWorld* World = GetWorld();
+				if (World->EndPhysicsTickFunction.IsTickFunctionRegistered() && SetupActorComponentTickFunction(&DeformableEndTickFunction))
 				{
-					DeformableEndTickFunction.AddPrerequisite(World, World->EndPhysicsTickFunction);
+					DeformableEndTickFunction.DeformableSolverComponent = this;
+					// Make sure our EndPhysicsTick gets called after physics simulation is finished
+					if (World != nullptr)
+					{
+						DeformableEndTickFunction.AddPrerequisite(this, PrimaryComponentTick);
+					}
 				}
 			}
-		}
-		else
-		{
-			DeformableEndTickFunction.UnRegisterTickFunction();
+			else
+			{
+				DeformableEndTickFunction.UnRegisterTickFunction();
+			}
 		}
 	}
-}
+	else if(DeformableEndTickFunction.IsTickFunctionRegistered())
+	{
+		DeformableEndTickFunction.UnRegisterTickFunction();
+	}
 
-//bool UDeformableSolverComponent::ShouldWaitForDeformableInTickFunction() const
-//{
-//	return bWaitForParallelDeformableTask || (CVarDeformablePhysicsTickWaitForParallelDeformableTask.GetValueOnAnyThread() != 0);
-//}
+}
 
 void UDeformableSolverComponent::BeginPlay()
 {
@@ -161,24 +196,25 @@ void UDeformableSolverComponent::TickComponent(float DeltaTime, enum ELevelTick
 	UE_LOG(LogDeformableSolverComponentInternal, Verbose, TEXT("UDeformableSolverComponent::TickComponent"));
 	TRACE_CPUPROFILER_EVENT_SCOPE(DeformableSolverComponent_TickComponent);
 
+	UpdateTickGroup();
+
 	UpdateDeformableEndTickState(IsSimulatable());
 
 	UpdateFromGameThread(DeltaTime);
 
 	if (bDoThreadedAdvance)
 	{
-		//ParallelDeformableTask = TGraphTask<FParallelDeformableTask>::CreateTask(nullptr, ENamedThreads::GameThread).ConstructAndDispatchWhenReady(this, DeltaTime);
-		//if (ShouldWaitForDeformableInTickFunction())
-		{ // see FParallelClothCompletionTask
-			FGraphEventArray Prerequisites;
-			Prerequisites.Add(ParallelDeformableTask);
-			FGraphEventRef DeformableCompletionEvent = TGraphTask<FParallelDeformableTask>::CreateTask(&Prerequisites, ENamedThreads::GameThread).ConstructAndDispatchWhenReady(this, DeltaTime);
-			ThisTickFunction->GetCompletionHandle()->DontCompleteUntil(DeformableCompletionEvent);
-		}
+		// see FParallelClothCompletionTask
+		FGraphEventArray Prerequisites;
+		Prerequisites.Add(ParallelDeformableTask);
+		FGraphEventRef DeformableCompletionEvent = TGraphTask<FParallelDeformableTask>::CreateTask(&Prerequisites, ENamedThreads::GameThread).ConstructAndDispatchWhenReady(this, DeltaTime);
+		ThisTickFunction->GetCompletionHandle()->DontCompleteUntil(DeformableCompletionEvent);
 	}
 	else
 	{
 		Simulate(DeltaTime);
+
+		UpdateFromSimulation(DeltaTime);
 	}
 
 
diff --git a/Engine/Plugins/Experimental/ChaosFlesh/Source/ChaosFleshEngine/Private/ChaosFlesh/FleshComponent.cpp b/Engine/Plugins/Experimental/ChaosFlesh/Source/ChaosFleshEngine/Private/ChaosFlesh/FleshComponent.cpp
--- a/Engine/Plugins/Experimental/ChaosFlesh/Source/ChaosFleshEngine/Private/ChaosFlesh/FleshComponent.cpp
+++ b/Engine/Plugins/Experimental/ChaosFlesh/Source/ChaosFleshEngine/Private/ChaosFlesh/FleshComponent.cpp
@@ -4,7 +4,10 @@
 
 #include "Animation/SkeletalMeshActor.h"
 #include "Animation/Skeleton.h"
+#include "Chaos/DebugDrawQueue.h"
 #include "ChaosFlesh/FleshCollection.h"
+#include "ChaosFlesh/ChaosDeformableSolverActor.h"
+#include "ChaosFlesh/ChaosDeformableSolverComponent.h"
 #include "Components/SkeletalMeshComponent.h"
 #include "Dataflow/DataflowEngineUtil.h"
 #include "Engine/SkeletalMesh.h"
@@ -20,6 +23,8 @@ UFleshComponent::UFleshComponent(const FObjectInitializer& ObjectInitializer)
 	: Super(ObjectInitializer)
 {
 	PrimaryComponentTick.bCanEverTick = true;
+	PrimaryComponentTick.TickGroup = TG_LastDemotable;
+
 	bTickInEditor = true;
 	DynamicCollection = ObjectInitializer.CreateDefaultSubobject<UFleshDynamicAsset>(this, TEXT("Flesh Dynamic Asset"));
 	Mesh = ObjectInitializer.CreateDefaultSubobject<UProceduralMeshComponent>(this, TEXT("Flesh Visualization Component"));
@@ -57,6 +62,13 @@ void UFleshComponent::BeginPlay()
 			PrimaryComponentTick.AddPrerequisite(SkeletalMeshComponent, SkeletalMeshComponent->PrimaryComponentTick);
 		}
 	}
+	if (PrimarySolver)
+	{
+		if (UDeformableSolverComponent* DeformableSolverComponent = PrimarySolver->GetDeformableSolverComponent())
+		{
+			PrimaryComponentTick.AddPrerequisite(DeformableSolverComponent, DeformableSolverComponent->PrimaryComponentTick);
+		}
+	}
 }
 
 void UFleshComponent::TickComponent(float DeltaTime, enum ELevelTick TickType, FActorComponentTickFunction *ThisTickFunction)
@@ -165,6 +177,7 @@ UDeformablePhysicsComponent::FDataMapValue UFleshComponent::NewDeformableData()
 
 void UFleshComponent::UpdateFromSimualtion(const FDataMapValue* SimualtionBuffer)
 {
+	FTransform CurrTransform = GetComponentTransform();
 	if (const FFleshThreadingProxy::FFleshOutputBuffer* FleshBuffer = (*SimualtionBuffer)->As<FFleshThreadingProxy::FFleshOutputBuffer>())
 	{
 		if (GetDynamicCollection())
@@ -179,9 +192,12 @@ void UFleshComponent::UpdateFromSimualtion(const FDataMapValue* SimualtionBuffer
 			for (int i = DynamicVertex.Num() - 1; i >= 0; i--)
 			{
 				DynamicVertex[i] = UEVertf(GetComponentTransform().InverseTransformPosition(UEVertd(SimulationVertex[i])));
+				//DynamicVertex[i] = UEVertf(PrevTransform.InverseTransformPosition(UEVertd(SimulationVertex[i])));
+				//Chaos::FDebugDrawQueue::GetInstance().DrawDebugPoint(UEVertd(SimulationVertex[i]), FColor::Red, false, -1.0f, 0, 5);
 			}
 		}
 	}
+	//PrevTransform = this->GetComponentTransform();
 }
 
 void UFleshComponent::UpdateLocalBounds()
diff --git a/Engine/Plugins/Experimental/ChaosFlesh/Source/ChaosFleshEngine/Public/ChaosFlesh/ChaosDeformableSolverComponent.h b/Engine/Plugins/Experimental/ChaosFlesh/Source/ChaosFleshEngine/Public/ChaosFlesh/ChaosDeformableSolverComponent.h
--- a/Engine/Plugins/Experimental/ChaosFlesh/Source/ChaosFleshEngine/Public/ChaosFlesh/ChaosDeformableSolverComponent.h
+++ b/Engine/Plugins/Experimental/ChaosFlesh/Source/ChaosFleshEngine/Public/ChaosFlesh/ChaosDeformableSolverComponent.h
@@ -5,6 +5,7 @@
 #include "Chaos/Deformable/ChaosDeformableSolverProxy.h"
 #include "Chaos/Deformable/ChaosDeformableSolverTypes.h"
 #include "ChaosFlesh/ChaosDeformableSolverThreading.h"
+#include "ChaosFlesh/ChaosDeformableTypes.h"
 #include "ChaosFlesh/FleshComponent.h"
 #include "Components/SceneComponent.h"
 #include "UObject/ObjectMacros.h"
@@ -29,6 +30,7 @@ public:
 	typedef Chaos::Softs::FDeformableSolver FDeformableSolver;
 
 	~UDeformableSolverComponent();
+	void UpdateTickGroup();
 
 #if WITH_EDITOR
 	TArray< TObjectPtr<AActor> > PreEditChangeDeformableActors;
@@ -57,7 +59,11 @@ public:
 	/* Properties */
 
 	UPROPERTY(EditAnywhere, BlueprintReadOnly, Category = Chaos)
-	TArray< TObjectPtr<AActor> > DeformableActors;
+		TArray< TObjectPtr<AActor> > DeformableActors;
+
+	/** ObjectType defines how to initialize the rigid objects state, Kinematic, Sleeping, Dynamic. */
+	UPROPERTY(EditAnywhere, BlueprintReadOnly , Category = Chaos)
+		EDeformableExecutionModel ExecutionModel = EDeformableExecutionModel::Chaos_Deformable_DuringPhysics;
 
 	UPROPERTY(EditAnywhere, Category = Chaos)
 		int32 NumSubSteps = 2;
diff --git a/Engine/Plugins/Experimental/ChaosFlesh/Source/ChaosFleshEngine/Public/ChaosFlesh/ChaosDeformableTypes.h b/Engine/Plugins/Experimental/ChaosFlesh/Source/ChaosFleshEngine/Public/ChaosFlesh/ChaosDeformableTypes.h
new file mode 100644
--- /dev/null
+++ b/Engine/Plugins/Experimental/ChaosFlesh/Source/ChaosFleshEngine/Public/ChaosFlesh/ChaosDeformableTypes.h
@@ -0,0 +1,16 @@
+// Copyright Epic Games, Inc. All Rights Reserved.
+
+#pragma once
+#include "CoreMinimal.h"
+#include "UObject/ObjectMacros.h"
+#include "Engine/EngineBaseTypes.h"
+
+UENUM(BlueprintType)
+enum class EDeformableExecutionModel : uint8
+{
+	Chaos_Deformable_PrePhysics UMETA(DisplayName = "Before Physics"),
+	Chaos_Deformable_DuringPhysics UMETA(DisplayName = "During Physics"),
+	Chaos_Deformable_PostPhysics UMETA(DisplayName = "After Physics"),
+	//
+	Chaos_Max UMETA(Hidden)
+};
\ No newline at end of file
diff --git a/Engine/Plugins/Experimental/ChaosFlesh/Source/ChaosFleshEngine/Public/ChaosFlesh/FleshAsset.h b/Engine/Plugins/Experimental/ChaosFlesh/Source/ChaosFleshEngine/Public/ChaosFlesh/FleshAsset.h
--- a/Engine/Plugins/Experimental/ChaosFlesh/Source/ChaosFleshEngine/Public/ChaosFlesh/FleshAsset.h
+++ b/Engine/Plugins/Experimental/ChaosFlesh/Source/ChaosFleshEngine/Public/ChaosFlesh/FleshAsset.h
@@ -7,12 +7,14 @@
 #include "ChaosFlesh/FleshCollection.h"
 #include "UObject/ObjectMacros.h"
 #include "Dataflow/DataflowNodeParameters.h"
+#include "Dataflow/DataflowEngineTypes.h"
 
 #include "FleshAsset.generated.h"
 
 class UFleshAsset;
 class UDataflow;
 class USkeletalMesh;
+class USkeleton;
 
 
 /**
@@ -103,12 +105,25 @@ public:
 	UPROPERTY(EditAnywhere, Category = "Dataflow")
 	FString DataflowTerminal = "FleshAssetTerminal";
 
+	UPROPERTY(EditAnywhere, Category = "Dataflow")
+	TArray<FStringValuePair> Overrides;
+
 	//
 	// SkeletalMesh
 	//
 	UPROPERTY(EditAnywhere, Category = "Animation")
 	TObjectPtr<USkeletalMesh> SkeletalMesh;
 
+	UPROPERTY(EditAnywhere, Category = "Animation")
+	TObjectPtr<USkeleton> Skeleton;
+
+	//
+	// SkeletalMesh
+	//
+	UPROPERTY(EditAnywhere, Category = "Geometry")
+	TObjectPtr<UStaticMesh> StaticMesh;
+
+
 
 #if WITH_EDITORONLY_DATA
 	UPROPERTY(EditAnywhere, Category = "Render")
diff --git a/Engine/Plugins/Experimental/ChaosFlesh/Source/ChaosFleshEngine/Public/ChaosFlesh/FleshComponent.h b/Engine/Plugins/Experimental/ChaosFlesh/Source/ChaosFleshEngine/Public/ChaosFlesh/FleshComponent.h
--- a/Engine/Plugins/Experimental/ChaosFlesh/Source/ChaosFleshEngine/Public/ChaosFlesh/FleshComponent.h
+++ b/Engine/Plugins/Experimental/ChaosFlesh/Source/ChaosFleshEngine/Public/ChaosFlesh/FleshComponent.h
@@ -82,5 +82,7 @@ private:
 
 	bool bBoundsNeedsUpdate = true;
 	FBoxSphereBounds BoundingBox = FBoxSphereBounds(ForceInitToZero);
+
+	FTransform PrevTransform = FTransform::Identity;
 };
 
diff --git a/Engine/Plugins/Experimental/ChaosFlesh/Source/ChaosFleshNodes/Private/Dataflow/ChaosFleshKinematicInitializationNodes.cpp b/Engine/Plugins/Experimental/ChaosFlesh/Source/ChaosFleshNodes/Private/Dataflow/ChaosFleshKinematicInitializationNodes.cpp
--- a/Engine/Plugins/Experimental/ChaosFlesh/Source/ChaosFleshNodes/Private/Dataflow/ChaosFleshKinematicInitializationNodes.cpp
+++ b/Engine/Plugins/Experimental/ChaosFlesh/Source/ChaosFleshNodes/Private/Dataflow/ChaosFleshKinematicInitializationNodes.cpp
@@ -184,6 +184,7 @@ void FKinematicInitializationDataflowNode::Evaluate(Dataflow::FContext& Context,
 							Kinematics.AddKinematicBinding(Kinematics.SetBoneBindings(Index, BoundVerts, BoundWeights));
 						}
 					}
+					GeometryCollection::Facades::FVertexBoneWeightsFacade(InCollection).AddBoneWeightsFromKinematicBindings();
 				}
 			}
 
diff --git a/Engine/Plugins/Experimental/Dataflow/Source/DataflowNodes/Private/Dataflow/DataflowSkeletalMeshNodes.cpp b/Engine/Plugins/Experimental/Dataflow/Source/DataflowNodes/Private/Dataflow/DataflowSkeletalMeshNodes.cpp
--- a/Engine/Plugins/Experimental/Dataflow/Source/DataflowNodes/Private/Dataflow/DataflowSkeletalMeshNodes.cpp
+++ b/Engine/Plugins/Experimental/Dataflow/Source/DataflowNodes/Private/Dataflow/DataflowSkeletalMeshNodes.cpp
@@ -3,6 +3,7 @@
 #include "Dataflow/DataflowSkeletalMeshNodes.h"
 #include "Dataflow/DataflowCore.h"
 #include "Dataflow/DataflowEngine.h"
+#include "Dataflow/DataflowEngineTypes.h"
 #include "Logging/LogMacros.h"
 #include "UObject/UnrealTypePrivate.h"
 #include "Animation/Skeleton.h"
@@ -73,7 +74,16 @@ void FSkeletalMeshBoneDataflowNode::Evaluate(Dataflow::FContext& Context, const
 
 		if( InDataType InSkeletalMesh = GetValue<InDataType>(Context, &SkeletalMesh) )
 		{
-			int32 Index = InSkeletalMesh->GetRefSkeleton().FindBoneIndex(BoneName);
+			FName LocalBoneName = BoneName;
+			if (LocalBoneName.IsNone())
+			{
+				if (const Dataflow::FEngineContext* EngineContext = Context.AsType<Dataflow::FEngineContext>())
+				{
+					LocalBoneName = FName(Dataflow::Reflection::FindOverrideProperty< FString >(EngineContext->Owner, PropertyName, FString("BoneName")));
+				}
+			}
+
+			int32 Index = InSkeletalMesh->GetRefSkeleton().FindBoneIndex(LocalBoneName);
 			SetValue<int>(Context, Index, &BoneIndexOut);
 		}
 
diff --git a/Engine/Plugins/Experimental/Dataflow/Source/DataflowNodes/Public/Dataflow/DataflowSkeletalMeshNodes.h b/Engine/Plugins/Experimental/Dataflow/Source/DataflowNodes/Public/Dataflow/DataflowSkeletalMeshNodes.h
--- a/Engine/Plugins/Experimental/Dataflow/Source/DataflowNodes/Public/Dataflow/DataflowSkeletalMeshNodes.h
+++ b/Engine/Plugins/Experimental/Dataflow/Source/DataflowNodes/Public/Dataflow/DataflowSkeletalMeshNodes.h
@@ -77,6 +77,9 @@ public:
 	UPROPERTY(meta = (DataflowOutput, DisplayName = "Index"))
 	int BoneIndexOut = INDEX_NONE;
 
+	UPROPERTY(EditAnywhere, Category = "Dataflow")
+	FName PropertyName = "Overrides";
+
 	FSkeletalMeshBoneDataflowNode(const Dataflow::FNodeParameters& InParam, FGuid InGuid = FGuid::NewGuid())
 		: FDataflowNode(InParam, InGuid)
 	{
diff --git a/Engine/Source/Runtime/Experimental/Chaos/Private/Chaos/Deformable/ChaosDeformableSolver.cpp b/Engine/Source/Runtime/Experimental/Chaos/Private/Chaos/Deformable/ChaosDeformableSolver.cpp
--- a/Engine/Source/Runtime/Experimental/Chaos/Private/Chaos/Deformable/ChaosDeformableSolver.cpp
+++ b/Engine/Source/Runtime/Experimental/Chaos/Private/Chaos/Deformable/ChaosDeformableSolver.cpp
@@ -4,6 +4,7 @@
 #include "Chaos/Deformable/ChaosDeformableSolverTypes.h"
 #include "Chaos/Deformable/ChaosDeformableSolverProxy.h"
 
+#include "Chaos/DebugDrawQueue.h"
 #include "Chaos/TriangleMesh.h"
 #include "Chaos/PBDAltitudeSpringConstraints.h"
 #include "Chaos/PBDBendingConstraints.h"
@@ -33,9 +34,16 @@
 
 DECLARE_CYCLE_STAT(TEXT("DeformableSolver.Advance"), STAT_DeformableSolver_Advance, STATGROUP_Chaos);
 
+
+
 DEFINE_LOG_CATEGORY_STATIC(LogChaosDeformableSolver, Log, All);
 namespace Chaos::Softs
 {
+	FDeformableDebugParams GDeformableDebugParams;
+
+	FAutoConsoleVariableRef CVarDeformableDebugParamsDrawTetrahedralParticles(TEXT("p.Chaos.DebugDraw.Deformable.TetrahedralParticle"), GDeformableDebugParams.bDoDrawTetrahedralParticles, TEXT("Debug draw the deformable solvers tetrahedron. [def: false]"));
+
+
 	FCriticalSection FDeformableSolver::InitializationMutex;
 	FCriticalSection FDeformableSolver::RemovalMutex;
 	FCriticalSection FDeformableSolver::PackageOutputMutex;
@@ -218,6 +226,28 @@ namespace Chaos::Softs
 		}
 	}
 
+	void FDeformableSolver::DebugDrawTetrahedralParticles(FFleshThreadingProxy& Proxy)
+	{
+#if WITH_EDITOR
+		auto ChaosTet = [](FIntVector4 V, int32 dp) { return Chaos::TVec4<int32>(dp + V.X, dp + V.Y, dp + V.Z, dp + V.W); };
+		auto ChaosVert = [](FVector3d V) { return Chaos::FVec3(V.X, V.Y, V.Z); };
+		auto DoubleVert = [](FVector3f V) { return FVector3d(V.X, V.Y, V.Z); };
+
+		const FIntVector2& Range = Proxy.GetSolverParticleRange();
+		const FManagedArrayCollection& Rest = Proxy.GetRestCollection();
+		const TManagedArray<FIntVector4>& Tetrahedron = Rest.GetAttribute<FIntVector4>("Tetrahedron", "Tetrahedral");
+		if (uint32 NumElements = Tetrahedron.Num())
+		{
+			const Chaos::Softs::FSolverParticles& P = Evolution->Particles();
+			for (uint32 edx = 0; edx < NumElements; ++edx)
+			{
+				auto T = ChaosTet(Tetrahedron[edx], Range[0]);
+				Chaos::FDebugDrawQueue::GetInstance().DrawDebugPoint(
+					DoubleVert(P.X(T[0])), FColor::Blue, false, -1.0f, 0, 5);
+			}
+		}
+#endif
+	}
 
 	void FDeformableSolver::InitializeTetrahedralConstraint(FFleshThreadingProxy& Proxy)
 	{
@@ -342,6 +372,7 @@ namespace Chaos::Softs
 						FTransform GlobalTransform = Proxy->GetInitialTransform();
 						const FIntVector2& Range = Proxy->GetSolverParticleRange();
 						const FManagedArrayCollection& Rest = Proxy->GetRestCollection();
+
 						if (Rest.FindAttributeTyped<FVector3f>("Vertex", FGeometryCollection::VerticesGroup))
 						{
 							const TManagedArray<FVector3f>& Vertex = Rest.GetAttribute<FVector3f>("Vertex", FGeometryCollection::VerticesGroup);
@@ -605,7 +636,27 @@ namespace Chaos::Softs
 			PushOutputPackage(Frame, MoveTemp(OutputBuffers));
 		}
 
-
+		{
+#if WITH_EDITOR
+			// debug draw
+	
+			//p.Chaos.DebugDraw.Enabled 1
+			if (GDeformableDebugParams.IsDebugDrawingEnabled())
+			{
+				for (TPair< FThreadingProxy::FKey, TUniquePtr<FThreadingProxy> >& BaseProxyPair : Proxies)
+				{
+					if (FFleshThreadingProxy* Proxy = BaseProxyPair.Value->As<FFleshThreadingProxy>())
+					{
+						if (GDeformableDebugParams.bDoDrawTetrahedralParticles)
+						{
+							//p.Chaos.DebugDraw.Deformable.TetrahedralParticles 1
+							DebugDrawTetrahedralParticles(*Proxy);
+						}
+					}
+				}
+			}
+#endif
+		}
 
 		EventPreBuffer.Broadcast(DeltaTime);
 	}
diff --git a/Engine/Source/Runtime/Experimental/Chaos/Private/GeometryCollection/Facades/CollectionVertexBoneWeightsFacade.cpp b/Engine/Source/Runtime/Experimental/Chaos/Private/GeometryCollection/Facades/CollectionVertexBoneWeightsFacade.cpp
--- a/Engine/Source/Runtime/Experimental/Chaos/Private/GeometryCollection/Facades/CollectionVertexBoneWeightsFacade.cpp
+++ b/Engine/Source/Runtime/Experimental/Chaos/Private/GeometryCollection/Facades/CollectionVertexBoneWeightsFacade.cpp
@@ -54,22 +54,53 @@ namespace GeometryCollection::Facades
 
 	//
 	//  Add Weights from Selection 
-	//  ... ... (Impressionist physbam)
 	//
 
 	void FVertexBoneWeightsFacade::AddBoneWeightsFromKinematicBindings() {
-		check(!IsConst()); DefineSchema(); if (IsValid()) {
-			TArray<float> Weights; TArray<int32> Indices;
+		check(!IsConst());
+		DefineSchema();
+
+		if (IsValid())
+		{
+			TArray<float> Weights;
+			TArray<int32> Indices;
+
 			int32 NumBones = ParentAttribute.Num(), NumVertices = BoneIndexAttribute.Num();
 			TManagedArray< TArray<int32> >& IndicesArray = BoneIndexAttribute.Modify();
 			TManagedArray< TArray<float> >& WeightsArray = BoneWeightAttribute.Modify();
+
 			GeometryCollection::Facades::FKinematicBindingFacade BindingFacade(ConstCollection);
-			for (int32 Kdx = BindingFacade.NumKinematicBindings() - 1; 0 <= Kdx; Kdx--) {
-				int32 Bone; TArray<int32> OutBoneVerts; TArray<float> OutBoneWeights;
+			for (int32 Kdx = BindingFacade.NumKinematicBindings() - 1; 0 <= Kdx; Kdx--)
+			{
+				int32 Bone;
+				TArray<int32> OutBoneVerts;
+				TArray<float> OutBoneWeights;
+
 				BindingFacade.GetBoneBindings(BindingFacade.GetKinematicBindingKey(Kdx), Bone, OutBoneVerts, OutBoneWeights);
-				if (0 <= Bone && Bone < NumBones) for (int32 Vdx = 0; Vdx < OutBoneVerts.Num(); Vdx++) {
-					int32 Vert = OutBoneVerts[Vdx]; float Weight = OutBoneWeights[Vdx];
-					if (0 <= Vert && Vert < NumVertices && !IndicesArray[Vert].Contains(Bone)) { IndicesArray[Vert].Add(Bone); WeightsArray[Vert].Add(Weight); }}}}}
+
+				if (0 <= Bone && Bone < NumBones)
+				{
+					for (int32 Vdx = 0; Vdx < OutBoneVerts.Num(); Vdx++)
+					{
+						int32 Vert = OutBoneVerts[Vdx]; float Weight = OutBoneWeights[Vdx];
+						if (0 <= Vert && Vert < NumVertices && !IndicesArray[Vert].Contains(Bone))
+						{
+							int32 BoneIndex = IndicesArray[Vert].Find(Bone);
+							if (BoneIndex == INDEX_NONE)
+							{
+								IndicesArray[Vert].Add(Bone);
+								WeightsArray[Vert].Add(Weight);
+							}
+							else if (0 <= BoneIndex && BoneIndex < WeightsArray[Vert].Num())
+							{
+								WeightsArray[Vert][BoneIndex] = Weight;
+							}
+						}
+					}
+				}
+			}
+		}
+	}
 
 };
 
diff --git a/Engine/Source/Runtime/Experimental/Chaos/Public/Chaos/Deformable/ChaosDeformableSolver.h b/Engine/Source/Runtime/Experimental/Chaos/Public/Chaos/Deformable/ChaosDeformableSolver.h
--- a/Engine/Source/Runtime/Experimental/Chaos/Public/Chaos/Deformable/ChaosDeformableSolver.h
+++ b/Engine/Source/Runtime/Experimental/Chaos/Public/Chaos/Deformable/ChaosDeformableSolver.h
@@ -130,6 +130,9 @@ namespace Chaos::Softs
 		void WriteTrisGEO(const FSolverParticles& Particles, const TArray<TVec3<int32>>& Mesh);
 		void WriteFrame(FThreadingProxy&, const FSolverReal DeltaTime);
 
+		/*Debug Drawing*/
+		void DebugDrawTetrahedralParticles(FFleshThreadingProxy& Proxy);
+
 		/*Game Thread API*/
 		bool HasObject(UObject* InObject) const { return InitializedObjects_External.Contains(InObject); }
 		void AddProxy(FThreadingProxy* InObject);
diff --git a/Engine/Source/Runtime/Experimental/Chaos/Public/Chaos/Deformable/ChaosDeformableSolverTypes.h b/Engine/Source/Runtime/Experimental/Chaos/Public/Chaos/Deformable/ChaosDeformableSolverTypes.h
--- a/Engine/Source/Runtime/Experimental/Chaos/Public/Chaos/Deformable/ChaosDeformableSolverTypes.h
+++ b/Engine/Source/Runtime/Experimental/Chaos/Public/Chaos/Deformable/ChaosDeformableSolverTypes.h
@@ -2,6 +2,7 @@
 
 #pragma once
 
+#include "Chaos/DebugDrawQueue.h"
 #include "Chaos/Deformable/ChaosDeformableSolverProxy.h"
 #include "Chaos/PBDSoftsEvolutionFwd.h"
 #include "CoreMinimal.h"
@@ -121,4 +122,21 @@ namespace Chaos::Softs
 	};
 
 
+	struct CHAOS_API FDeformableDebugParams
+	{				
+		bool bDoDrawTetrahedralParticles = false;
+
+
+		bool IsDebugDrawingEnabled()
+		{ 
+#if WITH_EDITOR
+			// p.Chaos.DebugDraw.Enabled 1
+			return Chaos::FDebugDrawQueue::GetInstance().IsDebugDrawingEnabled();
+#else
+			return false;
+#endif
+		}
+	};
+
+
 }; // namesapce Chaos::Softs
\ No newline at end of file
diff --git a/Engine/Source/Runtime/Experimental/Dataflow/Engine/Public/Dataflow/DataflowEngineTypes.h b/Engine/Source/Runtime/Experimental/Dataflow/Engine/Public/Dataflow/DataflowEngineTypes.h
new file mode 100644
--- /dev/null
+++ b/Engine/Source/Runtime/Experimental/Dataflow/Engine/Public/Dataflow/DataflowEngineTypes.h
@@ -0,0 +1,25 @@
+// Copyright Epic Games, Inc. All Rights Reserved.
+
+#pragma once
+
+#include "CoreMinimal.h"
+#include "UObject/ObjectMacros.h"
+#include "Engine/EngineBaseTypes.h"
+
+#include "DataflowEngineTypes.generated.h"
+
+USTRUCT(BlueprintType)
+struct DATAFLOWENGINE_API FStringValuePair
+{
+	GENERATED_BODY()
+
+	UPROPERTY(BlueprintReadWrite, EditAnywhere, Category = "Data")
+		FString Key;
+
+	UPROPERTY(BlueprintReadWrite, EditAnywhere, Category = "Data")
+		FString Value;
+
+	bool operator ==(const FStringValuePair& Other) const {
+		return Key.Equals(Other.Key) && Value.Equals(Other.Value);
+	}
+};
\ No newline at end of file
diff --git a/Engine/Source/Runtime/Experimental/Dataflow/Engine/Public/Dataflow/DataflowEngineUtil.h b/Engine/Source/Runtime/Experimental/Dataflow/Engine/Public/Dataflow/DataflowEngineUtil.h
--- a/Engine/Source/Runtime/Experimental/Dataflow/Engine/Public/Dataflow/DataflowEngineUtil.h
+++ b/Engine/Source/Runtime/Experimental/Dataflow/Engine/Public/Dataflow/DataflowEngineUtil.h
@@ -4,6 +4,7 @@
 
 #include "Containers/Array.h"
 #include "CoreMinimal.h"
+#include "Dataflow/DataflowEngineTypes.h"
 #include "Math/Transform.h"
 #include "ReferenceSkeleton.h"
 #include "Templates/Casts.h"
@@ -39,6 +40,34 @@ namespace Dataflow
 			}
 			return nullptr;
 		}
+
+		template<class T>
+		const T FindOverrideProperty(const UObject* Owner, FName PropertyName, FString ArrayKey, const T & Default = T())
+		{
+			if (Owner && Owner->GetClass())
+			{
+				if (FArrayProperty* ArrayProperty = FindFProperty<FArrayProperty>(Owner->GetClass(), PropertyName))
+				{
+					if (FStructProperty* ArrayInnerProperty = CastFieldChecked<FStructProperty>(ArrayProperty->Inner))
+					{
+						if (ArrayInnerProperty->Struct == FStringValuePair::StaticStruct())
+						{
+							if (const TArray<FStringValuePair>* ContainerData = ArrayProperty->ContainerPtrToValuePtr< TArray<FStringValuePair> >(Owner))
+							{
+								for (const FStringValuePair& Pair : *ContainerData)
+								{
+									if (Pair.Key.Equals(ArrayKey))
+									{
+										return Pair.Value;
+									}
+								}
+							}
+						}
+					}
+				}
+			}
+			return Default;
+		}
 	}
 
 	namespace Animation
