ChaosFlesh :
- Dataflow kinematic attachments to skeletal sub branch.
- Solver CVars for controlling activity.
- Debug draw kinematic particles.

diff --git a/Engine/Plugins/Experimental/ChaosFlesh/Source/ChaosFleshEngine/Private/ChaosFlesh/ChaosDeformableSolverComponent.cpp b/Engine/Plugins/Experimental/ChaosFlesh/Source/ChaosFleshEngine/Private/ChaosFlesh/ChaosDeformableSolverComponent.cpp
--- a/Engine/Plugins/Experimental/ChaosFlesh/Source/ChaosFleshEngine/Private/ChaosFlesh/ChaosDeformableSolverComponent.cpp
+++ b/Engine/Plugins/Experimental/ChaosFlesh/Source/ChaosFleshEngine/Private/ChaosFlesh/ChaosDeformableSolverComponent.cpp
@@ -13,6 +13,9 @@ static TAutoConsoleVariable<int32> CVarDeformablePhysicsTickWaitForParallelDefor
 	TEXT("If 1, always wait for deformable task completion in the Deformable Tick function. "\
 		 "If 0, wait at end - of - frame updates instead if allowed by component settings"));
 
+FChaosEngineDeformableCVarParams GChaosEngineDeformableCVarParams;
+FAutoConsoleVariableRef CVarChaosEngineDeformableSolverbEnabled(TEXT("p.Chaos.Deformable.EnableSimulation"), GChaosEngineDeformableCVarParams.bEnableDeformableSolver, TEXT("Enable the deformable simulation. [default : true]"));
+
 UDeformableSolverComponent::UDeformableSolverComponent(const FObjectInitializer& ObjectInitializer)
 	: Super(ObjectInitializer)
 	, Solver()
@@ -196,28 +199,29 @@ void UDeformableSolverComponent::TickComponent(float DeltaTime, enum ELevelTick
 	UE_LOG(LogDeformableSolverComponentInternal, Verbose, TEXT("UDeformableSolverComponent::TickComponent"));
 	TRACE_CPUPROFILER_EVENT_SCOPE(DeformableSolverComponent_TickComponent);
 
-	UpdateTickGroup();
+	if (GChaosEngineDeformableCVarParams.bEnableDeformableSolver)
+	{
+		UpdateTickGroup();
 
-	UpdateDeformableEndTickState(IsSimulatable());
+		UpdateDeformableEndTickState(IsSimulatable());
 
-	UpdateFromGameThread(DeltaTime);
+		UpdateFromGameThread(DeltaTime);
 
-	if (bDoThreadedAdvance)
-	{
-		// see FParallelClothCompletionTask
-		FGraphEventArray Prerequisites;
-		Prerequisites.Add(ParallelDeformableTask);
-		FGraphEventRef DeformableCompletionEvent = TGraphTask<FParallelDeformableTask>::CreateTask(&Prerequisites, ENamedThreads::GameThread).ConstructAndDispatchWhenReady(this, DeltaTime);
-		ThisTickFunction->GetCompletionHandle()->DontCompleteUntil(DeformableCompletionEvent);
-	}
-	else
-	{
-		Simulate(DeltaTime);
+		if (bDoThreadedAdvance)
+		{
+			// see FParallelClothCompletionTask
+			FGraphEventArray Prerequisites;
+			Prerequisites.Add(ParallelDeformableTask);
+			FGraphEventRef DeformableCompletionEvent = TGraphTask<FParallelDeformableTask>::CreateTask(&Prerequisites, ENamedThreads::GameThread).ConstructAndDispatchWhenReady(this, DeltaTime);
+			ThisTickFunction->GetCompletionHandle()->DontCompleteUntil(DeformableCompletionEvent);
+		}
+		else
+		{
+			Simulate(DeltaTime);
 
-		UpdateFromSimulation(DeltaTime);
+			UpdateFromSimulation(DeltaTime);
+		}
 	}
-
-
 }
 
 void UDeformableSolverComponent::EndPlay(const EEndPlayReason::Type EndPlayReason)
@@ -227,43 +231,44 @@ void UDeformableSolverComponent::EndPlay(const EEndPlayReason::Type EndPlayReaso
 
 void UDeformableSolverComponent::Reset()
 {
-	Solver.Reset(new FDeformableSolver({
-		NumSubSteps
-		, NumSolverIterations
-		, FixTimeStep
-		, TimeStepSize
-		, CacheToFile
-		, bEnableKinematics
-		, bUseFloor
-		, bDoSelfCollision
-		, bUseGridBasedConstraints
-		, GridDx
-		, bDoQuasistatics
-		, YoungModulus
-		, bDoBlended
-		, BlendedZeta
-		, Damping
-		, bEnableGravity
+	if (GChaosEngineDeformableCVarParams.bEnableDeformableSolver)
+	{
+		Solver.Reset(new FDeformableSolver({
+			NumSubSteps
+			, NumSolverIterations
+			, FixTimeStep
+			, TimeStepSize
+			, CacheToFile
+			, bEnableKinematics
+			, bUseFloor
+			, bDoSelfCollision
+			, bUseGridBasedConstraints
+			, GridDx
+			, bDoQuasistatics
+			, YoungModulus
+			, bDoBlended
+			, BlendedZeta
+			, Damping
+			, bEnableGravity
 		}));
 
-	
-	for (TObjectPtr<AActor>& DeformableActor : DeformableActors)
-	{
-		if (DeformableActor)
+		for (TObjectPtr<AActor>& DeformableActor : DeformableActors)
 		{
-			TArray<UDeformablePhysicsComponent*> DeformableComponents;
-			DeformableActor->GetComponents<UDeformablePhysicsComponent>(DeformableComponents);
-
-			for (UDeformablePhysicsComponent* DeformableComponent : DeformableComponents)
+			if (DeformableActor)
 			{
-				if (IsSimulating(DeformableComponent))
+				TArray<UDeformablePhysicsComponent*> DeformableComponents;
+				DeformableActor->GetComponents<UDeformablePhysicsComponent>(DeformableComponents);
+
+				for (UDeformablePhysicsComponent* DeformableComponent : DeformableComponents)
 				{
-					AddDeformableProxy(DeformableComponent);
+					if (IsSimulating(DeformableComponent))
+					{
+						AddDeformableProxy(DeformableComponent);
+					}
 				}
 			}
 		}
 	}
-	
 }
 
 /*
diff --git a/Engine/Plugins/Experimental/ChaosFlesh/Source/ChaosFleshEngine/Public/ChaosFlesh/ChaosDeformableTypes.h b/Engine/Plugins/Experimental/ChaosFlesh/Source/ChaosFleshEngine/Public/ChaosFlesh/ChaosDeformableTypes.h
--- a/Engine/Plugins/Experimental/ChaosFlesh/Source/ChaosFleshEngine/Public/ChaosFlesh/ChaosDeformableTypes.h
+++ b/Engine/Plugins/Experimental/ChaosFlesh/Source/ChaosFleshEngine/Public/ChaosFlesh/ChaosDeformableTypes.h
@@ -13,4 +13,9 @@ enum class EDeformableExecutionModel : uint8
 	Chaos_Deformable_PostPhysics UMETA(DisplayName = "After Physics"),
 	//
 	Chaos_Max UMETA(Hidden)
+};
+
+struct FChaosEngineDeformableCVarParams
+{
+	bool bEnableDeformableSolver = true;
 };
\ No newline at end of file
diff --git a/Engine/Plugins/Experimental/ChaosFlesh/Source/ChaosFleshNodes/Private/Dataflow/ChaosFleshKinematicInitializationNodes.cpp b/Engine/Plugins/Experimental/ChaosFlesh/Source/ChaosFleshNodes/Private/Dataflow/ChaosFleshKinematicInitializationNodes.cpp
--- a/Engine/Plugins/Experimental/ChaosFlesh/Source/ChaosFleshNodes/Private/Dataflow/ChaosFleshKinematicInitializationNodes.cpp
+++ b/Engine/Plugins/Experimental/ChaosFlesh/Source/ChaosFleshNodes/Private/Dataflow/ChaosFleshKinematicInitializationNodes.cpp
@@ -133,15 +133,15 @@ void FKinematicInitializationDataflowNode::Evaluate(Dataflow::FContext& Context,
 {
 	if (Out->IsA<DataType>(&Collection))
 	{
-		TArray<int32> BoundVerts;
-		TArray<float> BoundWeights;
-
 		DataType InCollection = GetValue<DataType>(Context, &Collection);
 
 		if (TManagedArray<FVector3f>* Vertices = InCollection.FindAttribute<FVector3f>("Vertex", FGeometryCollection::VerticesGroup))
 		{
 			if (FindInput(&VertexIndicesIn) && FindInput(&VertexIndicesIn)->GetConnection())
 			{
+				TArray<int32> BoundVerts;
+				TArray<float> BoundWeights;
+
 				for (int32 SelectionIndex : GetValue<TArray<int32>>(Context, &VertexIndicesIn))
 				{
 					if (0 <= SelectionIndex && SelectionIndex < Vertices->Num())
@@ -158,30 +158,70 @@ void FKinematicInitializationDataflowNode::Evaluate(Dataflow::FContext& Context,
 			}
 			else if (TObjectPtr<USkeletalMesh> SkeletalMesh = GetValue<TObjectPtr<USkeletalMesh>>(Context, &SkeletalMeshIn))
 			{
-				int32 Index = GetValue<int32>(Context, &BoneIndexIn);
-				if (Index != INDEX_NONE)
+				int32 IndexValue = GetValue<int32>(Context, &BoneIndexIn);
+				if (IndexValue != INDEX_NONE)
 				{
-
 					TArray<FTransform> ComponentPose;
 					Dataflow::Animation::GlobalTransforms(SkeletalMesh->GetRefSkeleton(), ComponentPose);
-					if (0 < Index && Index < ComponentPose.Num())
-					{
-						FVector3f BonePosition(ComponentPose[Index].GetTranslation());
 
-						int NumVertices = Vertices->Num();
-						for (int i = Vertices->Num() - 1; i > 0; i--)
+					TArray<int32> BranchIndices;
+					if (SkeletalSelectionMode == ESkeletalSeletionMode::Dataflow_SkeletalSelection_Branch)
+					{
+						TArray<int32> ToProcess;
+						int32 CurrentIndex = IndexValue;
+						while (SkeletalMesh->GetRefSkeleton().IsValidIndex(CurrentIndex))
 						{
-							if ((BonePosition - (*Vertices)[i]).Length() < Radius)
+							TArray<int32> Buffer;
+							SkeletalMesh->GetRefSkeleton().GetDirectChildBones(CurrentIndex, Buffer);
+
+							if (Buffer.Num())
+							{
+								ToProcess.Append(Buffer);
+							}
+
+							BranchIndices.Add(CurrentIndex);
+
+							CurrentIndex = INDEX_NONE;
+							if (ToProcess.Num())
 							{
-								BoundVerts.Add(i);
-								BoundWeights.Add(1.0);
+								CurrentIndex = ToProcess.Pop();
 							}
 						}
+					}
+					else // ESkeletalSeletionMode::Dataflow_SkeletalSelection_Single
+					{
+						BranchIndices.Add(IndexValue);
+					}
+
+					TSet<int32> ProcessedVertices;
+					for (const int32 Index : BranchIndices)
+					{
+						TArray<int32> BoundVerts;
+						TArray<float> BoundWeights;
 
-						if (BoundVerts.Num())
+						if (0 < Index && Index < ComponentPose.Num())
 						{
-							GeometryCollection::Facades::FKinematicBindingFacade Kinematics(InCollection);
-							Kinematics.AddKinematicBinding(Kinematics.SetBoneBindings(Index, BoundVerts, BoundWeights));
+							FVector3f BonePosition(ComponentPose[Index].GetTranslation());
+
+							int NumVertices = Vertices->Num();
+							for (int i = Vertices->Num() - 1; i > 0; i--)
+							{
+								if ((BonePosition - (*Vertices)[i]).Length() < Radius)
+								{
+									if (!ProcessedVertices.Contains(i))
+									{
+										ProcessedVertices.Add(i);
+										BoundVerts.Add(i);
+										BoundWeights.Add(1.0);
+									}
+								}
+							}
+
+							if (BoundVerts.Num())
+							{
+								GeometryCollection::Facades::FKinematicBindingFacade Kinematics(InCollection);
+								Kinematics.AddKinematicBinding(Kinematics.SetBoneBindings(Index, BoundVerts, BoundWeights));
+							}
 						}
 					}
 					GeometryCollection::Facades::FVertexBoneWeightsFacade(InCollection).AddBoneWeightsFromKinematicBindings();
diff --git a/Engine/Plugins/Experimental/ChaosFlesh/Source/ChaosFleshNodes/Public/Dataflow/ChaosFleshKinematicInitializationNodes.h b/Engine/Plugins/Experimental/ChaosFlesh/Source/ChaosFleshNodes/Public/Dataflow/ChaosFleshKinematicInitializationNodes.h
--- a/Engine/Plugins/Experimental/ChaosFlesh/Source/ChaosFleshNodes/Public/Dataflow/ChaosFleshKinematicInitializationNodes.h
+++ b/Engine/Plugins/Experimental/ChaosFlesh/Source/ChaosFleshNodes/Public/Dataflow/ChaosFleshKinematicInitializationNodes.h
@@ -40,6 +40,15 @@ public:
 	virtual void Evaluate(Dataflow::FContext& Context, const FDataflowOutput* Out) const override;
 };
 
+UENUM(BlueprintType)
+enum class ESkeletalSeletionMode : uint8
+{
+	Dataflow_SkeletalSelection_Single UMETA(DisplayName = "Single"),
+	Dataflow_SkeletalSelection_Branch UMETA(DisplayName = "Sub-Branch"),
+	//
+	Chaos_Max UMETA(Hidden)
+};
+
 
 USTRUCT(meta = (DataflowFlesh))
 struct FKinematicInitializationDataflowNode : public FDataflowNode
@@ -57,6 +66,9 @@ public:
 	UPROPERTY(EditAnywhere, Category = "Dataflow")
 	FTransform Transform;
 
+	UPROPERTY(EditAnywhere, Category = "Dataflow")
+	ESkeletalSeletionMode SkeletalSelectionMode = ESkeletalSeletionMode::Dataflow_SkeletalSelection_Single;
+
 	UPROPERTY(meta = (DataflowInput, DataflowOutput, DisplayName = "Collection", DataflowPassthrough  = "Collection"))
 	FManagedArrayCollection Collection;
 
@@ -69,6 +81,7 @@ public:
 	UPROPERTY(meta = (DataflowInput, DisplayName = "BoneIndex"))
 	int32 BoneIndexIn = INDEX_NONE;
 	
+
 	FKinematicInitializationDataflowNode(const Dataflow::FNodeParameters& InParam, FGuid InGuid = FGuid::NewGuid())
 		: FDataflowNode(InParam, InGuid)
 	{
diff --git a/Engine/Source/Runtime/Experimental/Chaos/Private/Chaos/Deformable/ChaosDeformableSolver.cpp b/Engine/Source/Runtime/Experimental/Chaos/Private/Chaos/Deformable/ChaosDeformableSolver.cpp
--- a/Engine/Source/Runtime/Experimental/Chaos/Private/Chaos/Deformable/ChaosDeformableSolver.cpp
+++ b/Engine/Source/Runtime/Experimental/Chaos/Private/Chaos/Deformable/ChaosDeformableSolver.cpp
@@ -42,6 +42,7 @@ namespace Chaos::Softs
 	FDeformableDebugParams GDeformableDebugParams;
 
 	FAutoConsoleVariableRef CVarDeformableDebugParamsDrawTetrahedralParticles(TEXT("p.Chaos.DebugDraw.Deformable.TetrahedralParticle"), GDeformableDebugParams.bDoDrawTetrahedralParticles, TEXT("Debug draw the deformable solvers tetrahedron. [def: false]"));
+	FAutoConsoleVariableRef CVarDeformableDebugParamsDrawKinematicParticles(TEXT("p.Chaos.DebugDraw.Deformable.KinematicParticle"), GDeformableDebugParams.bDoDrawKinematicParticles, TEXT("Debug draw the deformables kinematic particles. [def: false]"));
 
 
 	FCriticalSection FDeformableSolver::InitializationMutex;
@@ -381,8 +382,6 @@ namespace Chaos::Softs
 
 	void FDeformableSolver::InitializeKinematicConstraint()
 	{
-
-
 		auto MKineticUpdate = [this](FSolverParticles& MParticles, const FSolverReal Dt, const FSolverReal MTime, const int32 Index)
 		{
 			TRACE_CPUPROFILER_EVENT_SCOPE(DeformableSolver_MKineticUpdate);
@@ -457,6 +456,18 @@ namespace Chaos::Softs
 
 											MParticles.X(Index) = GlobalTransform.TransformPosition(ComponentPointAtT);
 
+
+#if WITH_EDITOR
+											//debug draw
+											//p.Chaos.DebugDraw.Enabled 1
+											//p.Chaos.DebugDraw.Deformable.KinematicParticle 1
+											if (GDeformableDebugParams.IsDebugDrawingEnabled() && GDeformableDebugParams.bDoDrawKinematicParticles)
+											{
+												auto DoubleVert = [](FVector3f V) { return FVector3d(V.X, V.Y, V.Z); };
+												Chaos::FDebugDrawQueue::GetInstance().DrawDebugPoint(DoubleVert(MParticles.X(Index)), FColor::Red, false, -1.0f, 0, 5);
+											}
+#endif
+
 											// @todo(flesh): Add non rigid skinning weights. 
 											// Currently this just grabs the last joint, ideally
 											// we would be using the weighted average for the vertex
diff --git a/Engine/Source/Runtime/Experimental/Chaos/Public/Chaos/Deformable/ChaosDeformableSolverTypes.h b/Engine/Source/Runtime/Experimental/Chaos/Public/Chaos/Deformable/ChaosDeformableSolverTypes.h
--- a/Engine/Source/Runtime/Experimental/Chaos/Public/Chaos/Deformable/ChaosDeformableSolverTypes.h
+++ b/Engine/Source/Runtime/Experimental/Chaos/Public/Chaos/Deformable/ChaosDeformableSolverTypes.h
@@ -133,6 +133,7 @@ namespace Chaos::Softs
 	struct CHAOS_API FDeformableDebugParams
 	{				
 		bool bDoDrawTetrahedralParticles = false;
+		bool bDoDrawKinematicParticles = false;
 
 
 		bool IsDebugDrawingEnabled()
