Add support for creating tetrahedral meshes directly from SkeletalMesh.

diff --git a/Engine/Plugins/Experimental/ChaosFlesh/Source/ChaosFleshNodes/ChaosFleshNodes.Build.cs b/Engine/Plugins/Experimental/ChaosFlesh/Source/ChaosFleshNodes/ChaosFleshNodes.Build.cs
--- a/Engine/Plugins/Experimental/ChaosFlesh/Source/ChaosFleshNodes/ChaosFleshNodes.Build.cs
+++ b/Engine/Plugins/Experimental/ChaosFlesh/Source/ChaosFleshNodes/ChaosFleshNodes.Build.cs
@@ -21,6 +21,8 @@ namespace UnrealBuildTool.Rules
 					"GeometryAlgorithms",
 					"GeometryCore",
 					"MeshConversion",
+					"MeshConversionEngineTypes",
+					"MeshDescription",
 					"TetMeshing"
 				}
 			);
diff --git a/Engine/Plugins/Experimental/ChaosFlesh/Source/ChaosFleshNodes/Private/Dataflow/ChaosFleshTetrahedralNodes.cpp b/Engine/Plugins/Experimental/ChaosFlesh/Source/ChaosFleshNodes/Private/Dataflow/ChaosFleshTetrahedralNodes.cpp
--- a/Engine/Plugins/Experimental/ChaosFlesh/Source/ChaosFleshNodes/Private/Dataflow/ChaosFleshTetrahedralNodes.cpp
+++ b/Engine/Plugins/Experimental/ChaosFlesh/Source/ChaosFleshNodes/Private/Dataflow/ChaosFleshTetrahedralNodes.cpp
@@ -13,11 +13,16 @@
 #include "Dataflow/DataflowInputOutput.h"
 #include "DynamicMesh/DynamicMesh3.h"
 #include "DynamicMesh/DynamicMeshAABBTree3.h"
+#include "Engine/SkeletalMesh.h"
 #include "Engine/StaticMesh.h"
 #include "FTetWildWrapper.h"
 #include "Generate/IsosurfaceStuffing.h"
 #include "GeometryCollection/ManagedArrayCollection.h"
+#include "MeshDescription.h"
 #include "MeshDescriptionToDynamicMesh.h"
+#include "Rendering/SkeletalMeshLODImporterData.h"
+#include "Rendering/SkeletalMeshModel.h"
+#include "SkeletalMeshLODModelToDynamicMesh.h" // MeshModelingBlueprints
 #include "Spatial/FastWinding.h"
 #include "Spatial/MeshAABBTree3.h"
 
@@ -108,24 +113,63 @@ void FGenerateTetrahedralCollectionDataflowNodes::Evaluate(Dataflow::FContext& C
 	{
 		TUniquePtr<FFleshCollection> InCollection(GetValue<DataType>(Context, &Collection).NewCopy<FFleshCollection>());
 		TObjectPtr<const UStaticMesh> InStaticMesh(GetValue<TObjectPtr<const UStaticMesh>>(Context, &StaticMesh));
+		TObjectPtr<const USkeletalMesh> InSkeletalMesh(FindInput(&SkeletalMesh) ? GetValue<TObjectPtr<const USkeletalMesh>>(Context, &SkeletalMesh) : nullptr);
 
-#if WITH_EDITORONLY_DATA
-		if (Method == TetMeshingMethod::IsoStuffing)
-		{
-			EvaluateIsoStuffing(Context, InCollection, InStaticMesh);
-		}
-		else if (Method == TetMeshingMethod::TetWild)
-		{
-			EvaluateTetWild(Context, InCollection, InStaticMesh);
-		}
-		else
+		if (InStaticMesh || InSkeletalMesh)
 		{
-			ensureMsgf(false, TEXT("FGenerateTetrahedralCollectionDataflowNodes unsupported Method."));
-		}
+#if WITH_EDITORONLY_DATA
+			UE::Geometry::FDynamicMesh3 DynamicMesh;
+			if (InStaticMesh)
+			{
+				// make a mesh description for UE::Geometry tools
+				FMeshDescriptionToDynamicMesh GetSourceMesh;
+				bool bUsingHiResSource = InStaticMesh->IsHiResMeshDescriptionValid();
+				const FMeshDescription* UseSourceMeshDescription =
+					(bUsingHiResSource) ? InStaticMesh->GetHiResMeshDescription() : InStaticMesh->GetMeshDescription(0);
+				GetSourceMesh.Convert(UseSourceMeshDescription, DynamicMesh);
+			}
+			else if (InSkeletalMesh)
+			{
+				// Check first if we have bulk data available and non-empty.
+				const int32 LODIndex = 0;
+				FMeshDescription SourceMesh;
+				if (InSkeletalMesh->IsLODImportedDataBuildAvailable(LODIndex) && !InSkeletalMesh->IsLODImportedDataEmpty(LODIndex))
+				{
+					FSkeletalMeshImportData SkeletalMeshImportData;
+					InSkeletalMesh->LoadLODImportedData(LODIndex, SkeletalMeshImportData);
+					SkeletalMeshImportData.GetMeshDescription(SourceMesh);
+				}
+				else
+				{
+					// Fall back on the LOD model directly if no bulk data exists. When we commit
+					// the mesh description, we override using the bulk data. This can happen for older
+					// skeletal meshes, from UE 4.24 and earlier.
+					const FSkeletalMeshModel* SkeletalMeshModel = InSkeletalMesh->GetImportedModel();
+					if (SkeletalMeshModel && SkeletalMeshModel->LODModels.IsValidIndex(LODIndex))
+					{
+						SkeletalMeshModel->LODModels[LODIndex].GetMeshDescription(SourceMesh, InSkeletalMesh);
+					}
+				}
+				FMeshDescriptionToDynamicMesh Converter;
+				Converter.Convert(&SourceMesh, DynamicMesh);
+			}
+
+			if (Method == TetMeshingMethod::IsoStuffing)
+			{
+				EvaluateIsoStuffing(Context, InCollection, DynamicMesh);
+			}
+			else if (Method == TetMeshingMethod::TetWild)
+			{
+				EvaluateTetWild(Context, InCollection, DynamicMesh);
+			}
+			else
+			{
+				ensureMsgf(false, TEXT("FGenerateTetrahedralCollectionDataflowNodes unsupported Method."));
+			}
 #else
-		ensureMsgf(false,TEXT("FGenerateTetrahedralCollectionDataflowNodes is an editor only node."));
+			ensureMsgf(false, TEXT("FGenerateTetrahedralCollectionDataflowNodes is an editor only node."));
 #endif
-
+		} // end if InStaticMesh || InSkeletalMesh
 		SetValue<DataType>(Context, *(FManagedArrayCollection*)InCollection.Get(), &Collection);
 	}
 }
@@ -133,19 +177,11 @@ void FGenerateTetrahedralCollectionDataflowNodes::Evaluate(Dataflow::FContext& C
 void FGenerateTetrahedralCollectionDataflowNodes::EvaluateIsoStuffing(
 	Dataflow::FContext& Context, 
 	TUniquePtr<FFleshCollection>& InCollection,
-	TObjectPtr<const UStaticMesh>& InStaticMesh) const
+	const UE::Geometry::FDynamicMesh3& DynamicMesh) const
 {
 #if WITH_EDITORONLY_DATA
-	if (InStaticMesh && NumCells > 0 && (-.5 <= OffsetPercent && OffsetPercent <= 0.5))
+	if (NumCells > 0 && (-.5 <= OffsetPercent && OffsetPercent <= 0.5))
 	{
-		// make a mesh description for UE::Geometry tools
-		UE::Geometry::FDynamicMesh3 DynamicMesh;
-		FMeshDescriptionToDynamicMesh GetSourceMesh;
-		bool bUsingHiResSource = InStaticMesh->IsHiResMeshDescriptionValid();
-		const FMeshDescription* UseSourceMeshDescription =
-			(bUsingHiResSource) ? InStaticMesh->GetHiResMeshDescription() : InStaticMesh->GetMeshDescription(0);
-		GetSourceMesh.Convert(UseSourceMeshDescription, DynamicMesh);
-
 		// Tet mesh generation
 		UE::Geometry::TIsosurfaceStuffing<double> IsosurfaceStuffing;
 		UE::Geometry::FDynamicMeshAABBTree3 Spatial(&DynamicMesh);
@@ -198,19 +234,11 @@ void FGenerateTetrahedralCollectionDataflowNodes::EvaluateIsoStuffing(
 void FGenerateTetrahedralCollectionDataflowNodes::EvaluateTetWild(
 	Dataflow::FContext& Context, 
 	TUniquePtr<FFleshCollection>& InCollection,
-	TObjectPtr<const UStaticMesh>& InStaticMesh) const
+	const UE::Geometry::FDynamicMesh3& DynamicMesh) const
 {
 #if WITH_EDITORONLY_DATA
-	if (StaticMesh && InStaticMesh)
+	if (/* placeholder for conditions for exec */true)
 	{
-		// make a mesh description for UE::Geometry tools
-		UE::Geometry::FDynamicMesh3 DynamicMesh;
-		FMeshDescriptionToDynamicMesh GetSourceMesh;
-		bool bUsingHiResSource = InStaticMesh->IsHiResMeshDescriptionValid();
-		const FMeshDescription* UseSourceMeshDescription =
-			(bUsingHiResSource) ? InStaticMesh->GetHiResMeshDescription() : InStaticMesh->GetMeshDescription(0);
-		GetSourceMesh.Convert(UseSourceMeshDescription, DynamicMesh);
-
 		// Pull out Vertices and Triangles
 		TArray<FVector> Verts;
 		TArray<FIntVector3> Tris;
diff --git a/Engine/Plugins/Experimental/ChaosFlesh/Source/ChaosFleshNodes/Public/Dataflow/ChaosFleshTetrahedralNodes.h b/Engine/Plugins/Experimental/ChaosFlesh/Source/ChaosFleshNodes/Public/Dataflow/ChaosFleshTetrahedralNodes.h
--- a/Engine/Plugins/Experimental/ChaosFlesh/Source/ChaosFleshNodes/Public/Dataflow/ChaosFleshTetrahedralNodes.h
+++ b/Engine/Plugins/Experimental/ChaosFlesh/Source/ChaosFleshNodes/Public/Dataflow/ChaosFleshTetrahedralNodes.h
@@ -8,8 +8,14 @@
 
 #include "ChaosFleshTetrahedralNodes.generated.h"
 
+class USkeletalMesh;
 class UStaticMesh;
 class FFleshCollection;
+namespace UE {
+	namespace Geometry {
+		class FDynamicMesh3;
+	}
+}
 
 USTRUCT(meta = (DataflowFlesh))
 struct FConstructTetGridNode : public FDataflowNode
@@ -114,6 +120,9 @@ public:
 	UPROPERTY(EditAnywhere, Category = "Dataflow")
 	TObjectPtr<const UStaticMesh> StaticMesh = nullptr;
 
+	UPROPERTY(EditAnywhere, Category = "Dataflow")
+	TObjectPtr<const USkeletalMesh> SkeletalMesh = nullptr;
+
 	UPROPERTY(meta = (DataflowInput, DataflowOutput, DisplayName = "Collection"))
 	FManagedArrayCollection Collection;
 
@@ -121,6 +130,7 @@ public:
 		: FDataflowNode(InParam, InGuid)
 	{
 		RegisterInputConnection(&StaticMesh);
+		RegisterInputConnection(&SkeletalMesh);
 		RegisterInputConnection(&Collection);
 		RegisterOutputConnection(&Collection);
 	}
@@ -128,8 +138,8 @@ public:
 	virtual void Evaluate(Dataflow::FContext& Context, const FDataflowOutput* Out) const override;
 
 protected:
-	void EvaluateIsoStuffing(Dataflow::FContext& Context, TUniquePtr<FFleshCollection>& InCollection, TObjectPtr<const UStaticMesh>& InStaticMesh) const;
-	void EvaluateTetWild(Dataflow::FContext& Context, TUniquePtr<FFleshCollection>& InCollection, TObjectPtr<const UStaticMesh>& InStaticMesh) const;
+	void EvaluateIsoStuffing(Dataflow::FContext& Context, TUniquePtr<FFleshCollection>& InCollection, const UE::Geometry::FDynamicMesh3& DynamicMesh) const;
+	void EvaluateTetWild(Dataflow::FContext& Context, TUniquePtr<FFleshCollection>& InCollection, const UE::Geometry::FDynamicMesh3& DynamicMesh) const;
 };
 
 namespace Dataflow
