Flesh deformer data interface and shader.

diff --git a/Engine/Plugins/Experimental/ChaosFlesh/ChaosFlesh.uplugin b/Engine/Plugins/Experimental/ChaosFlesh/ChaosFlesh.uplugin
--- a/Engine/Plugins/Experimental/ChaosFlesh/ChaosFlesh.uplugin
+++ b/Engine/Plugins/Experimental/ChaosFlesh/ChaosFlesh.uplugin
@@ -53,6 +53,14 @@
 		{
 			"Name": "GeometryProcessing",
 			"Enabled": true
+		}, 
+		{
+			"Name": "ComputeFramework",
+			"Enabled": true
+		}, 
+		{
+			"Name": "DeformerGraph",
+			"Enabled": true
 		}
 	]
 }
diff --git a/Engine/Plugins/Experimental/ChaosFlesh/Shaders/Private/DataInterfaceFlesh.ush b/Engine/Plugins/Experimental/ChaosFlesh/Shaders/Private/DataInterfaceFlesh.ush
new file mode 100644
--- /dev/null
+++ b/Engine/Plugins/Experimental/ChaosFlesh/Shaders/Private/DataInterfaceFlesh.ush
@@ -0,0 +1,26 @@
+// Copyright Epic Games, Inc.  All Rights Reserved.
+
+#include "/Plugin/ComputeFramework/Private/ComputeKernelCommon.ush"
+
+float4x4 DI_LOCAL(SimToLocal) =              // Simulation space (world?) to local
+         { 1.f, 0.f, 0.f, 0.f,
+           0.f, 1.f, 0.f, 0.f,
+           0.f, 0.f, 1.f, 0.f,
+           0.f, 0.f, 0.f, 1.f };
+
+Buffer<float3> DI_LOCAL(SimPositions);      // Current tetrahedra positions
+Buffer<int> DI_LOCAL(Tetrahedra);           // Tet mesh topology
+
+Buffer<int> DI_LOCAL(TetBinding);           // Per vertex tetrahedron index
+Buffer<float4> DI_LOCAL(TetBaryWeights);    // Per vertex tetrahedron weights
+
+float3 DI_LOCAL(GetEmbeddedTetPosition)(uint Index)
+{
+    int TetIdx = DI_LOCAL(TetBinding)[Index];
+    float4 TetWeights = DI_LOCAL(TetBaryWeights)[TetIdx];
+    return
+        DI_LOCAL(SimPositions)[DI_LOCAL(Tetrahedra[TetIdx*4+0])] * TetWeights[0] +
+        DI_LOCAL(SimPositions)[DI_LOCAL(Tetrahedra[TetIdx*4+1])] * TetWeights[1] +
+        DI_LOCAL(SimPositions)[DI_LOCAL(Tetrahedra[TetIdx*4+2])] * TetWeights[2] +
+        DI_LOCAL(SimPositions)[DI_LOCAL(Tetrahedra[TetIdx*4+3])] * TetWeights[3];
+}
\ No newline at end of file
diff --git a/Engine/Plugins/Experimental/ChaosFlesh/Source/ChaosFleshEngine/ChaosFleshEngine.Build.cs b/Engine/Plugins/Experimental/ChaosFlesh/Source/ChaosFleshEngine/ChaosFleshEngine.Build.cs
--- a/Engine/Plugins/Experimental/ChaosFlesh/Source/ChaosFleshEngine/ChaosFleshEngine.Build.cs
+++ b/Engine/Plugins/Experimental/ChaosFlesh/Source/ChaosFleshEngine/ChaosFleshEngine.Build.cs
@@ -8,26 +8,54 @@ namespace UnrealBuildTool.Rules
 		{
 			SetupModulePhysicsSupport(Target);
 
+			PublicIncludePaths.AddRange(
+				new string[] {
+					// ... add public include paths required here ...
+				}
+				);
+
+			PrivateIncludePaths.AddRange(
+				new string[] {
+					// ... add other private include paths required here ...
+				}
+				);
+
 			PublicDependencyModuleNames.AddRange(
 				new string[]
 				{
 					"Core",
+					"ProceduralMeshComponent",
+					// ... add other public dependencies that you statically link with here ...
+				}
+				);
+
+			PrivateDependencyModuleNames.AddRange(
+				new string[]
+				{
+					"ComputeFramework",
 					"CoreUObject",
 					"Chaos",
-					"Engine",
-                    "RenderCore",
-                    "RHI",
-					"Renderer",
-                    "FieldSystemEngine",
-	                "ChaosFlesh",
-					"NetCore",
-					"ProceduralMeshComponent",
+					"ChaosCaching",
+					"ChaosFlesh",
 					"DataflowCore",
 					"DataflowEngine",
-					"ChaosCaching",
+					"Engine",
+					"FieldSystemEngine",
+					"NetCore",
+					"OptimusCore",
+					"Projects",
+					"RenderCore",
+                    "RHI",
+					"Renderer"
 				}
 				);
 
+			DynamicallyLoadedModuleNames.AddRange(
+				new string[]
+				{
+					// ... add any modules that your module loads dynamically here ...
+				}
+				);
 
 			PrivateDefinitions.Add("CHAOS_INCLUDE_LEVEL_1=1");
 		}
diff --git a/Engine/Plugins/Experimental/ChaosFlesh/Source/ChaosFleshEngine/Private/ChaosFlesh/ChaosFleshEnginePlugin.cpp b/Engine/Plugins/Experimental/ChaosFlesh/Source/ChaosFleshEngine/Private/ChaosFlesh/ChaosFleshEnginePlugin.cpp
--- a/Engine/Plugins/Experimental/ChaosFlesh/Source/ChaosFleshEngine/Private/ChaosFlesh/ChaosFleshEnginePlugin.cpp
+++ b/Engine/Plugins/Experimental/ChaosFlesh/Source/ChaosFleshEngine/Private/ChaosFlesh/ChaosFleshEnginePlugin.cpp
@@ -3,7 +3,10 @@
 #include "ChaosFlesh/ChaosFleshEnginePlugin.h"
 
 #include "ChaosCache/FleshComponentCacheAdapter.h"
-
+#include "DataInterfaces/DIFleshDeformer.h"
+#include "Interfaces/IPluginManager.h"
+#include "IOptimusCoreModule.h"
+#include "ShaderCore.h"
 
 
 class FChaosFleshEnginePlugin : public IChaosFleshEnginePlugin
@@ -19,14 +22,27 @@ private:
 
 void FChaosFleshEnginePlugin::StartupModule()
 {
+	// Make sure our shaders can be found via the virtual shader paths.
+	const FString PluginShaderDir = 
+		FPaths::Combine(
+			IPluginManager::Get().FindPlugin(TEXT("ChaosFlesh"))->GetBaseDir(), 
+			TEXT("Source/ChaosFleshEngine/Shaders/Private"));
+	AddShaderSourceDirectoryMapping(TEXT("/Plugin/ChaosFleshEngine"), PluginShaderDir);
+
+	// Self-register for now.
+	// FIXME: Should add a method to IOptimusModule to register interfaces.
+	if (IOptimusCoreModule* OptimusCoreModule = FModuleManager::GetModulePtr<IOptimusCoreModule>("OptimusCore"))
+	{
+		OptimusCoreModule->RegisterDataInterfaceClass<UDIFleshDeformer>();
+	}
+
 	FleshCacheAdapter = MakeUnique<Chaos::FFleshCacheAdapter>();
 	Chaos::RegisterAdapter(FleshCacheAdapter.Get());
 }
 
 
 void FChaosFleshEnginePlugin::ShutdownModule()
-{
-	
+{	
 }
 
 
diff --git a/Engine/Plugins/Experimental/ChaosFlesh/Source/ChaosFleshEngine/Private/DataInterfaces/DIFleshDeformer.cpp b/Engine/Plugins/Experimental/ChaosFlesh/Source/ChaosFleshEngine/Private/DataInterfaces/DIFleshDeformer.cpp
new file mode 100644
--- /dev/null
+++ b/Engine/Plugins/Experimental/ChaosFlesh/Source/ChaosFleshEngine/Private/DataInterfaces/DIFleshDeformer.cpp
@@ -0,0 +1,525 @@
+// Copyright Epic Games, Inc. All Rights Reserved.
+
+#include "DataInterfaces/DIFleshDeformer.h"
+
+#include "Components/SkeletalMeshComponent.h"
+#include "Containers/Set.h"
+#include "ComputeFramework/ComputeKernelPermutationVector.h"
+#include "ComputeFramework/ShaderParamTypeDefinition.h"
+#include "GenericPlatform/GenericPlatformMath.h"
+#include "GeometryCollection/Facades/CollectionTetrahedralBindingsFacade.h"
+#include "OptimusDataDomain.h"
+#include "ShaderParameterMetadataBuilder.h"
+#include "RenderGraphBuilder.h"
+#include "Rendering/SkeletalMeshLODRenderData.h"
+#include "Rendering/SkeletalMeshModel.h"
+#include "Rendering/SkeletalMeshRenderData.h"
+#include "ShaderCompilerCore.h"
+#include "SkeletalMeshDeformerHelpers.h"
+#include "SkeletalRenderPublic.h"
+
+DEFINE_LOG_CATEGORY(LogFleshDeformer);
+
+//
+// Interface
+//
+
+FString
+UDIFleshDeformer::GetDisplayName() const
+{
+	return TEXT("Flesh Deformer");
+}
+
+TArray<FOptimusCDIPinDefinition>
+UDIFleshDeformer::GetPinDefinitions() const
+{
+	TArray<FOptimusCDIPinDefinition> Defs;
+	Defs.Add({ "Amount", "ReadAmount" });
+	Defs.Add({ "Mask", "ReadMask", Optimus::DomainName::Vertex, "ReadNumVertices" });
+	Defs.Add({ "Parents", "ReadParents", Optimus::DomainName::Vertex, "ReadNumVertices" });
+	Defs.Add({ "Weights", "ReadWeights", Optimus::DomainName::Vertex, "ReadNumVertices"  });
+	Defs.Add({ "Offset", "ReadOffset", Optimus::DomainName::Vertex, "ReadNumVertices" });
+	Defs.Add({ "EmbeddedPos", "GetEmbeddedPos", Optimus::DomainName::Vertex, "ReadNumVertices" });
+	return Defs;
+}
+
+TSubclassOf<UActorComponent> 
+UDIFleshDeformer::GetRequiredComponentClass() const
+{
+	return USkinnedMeshComponent::StaticClass();
+}
+
+void
+UDIFleshDeformer::GetSupportedInputs(TArray<FShaderFunctionDefinition>& OutFunctions) const
+{
+	OutFunctions.AddDefaulted_GetRef()
+		.SetName(TEXT("ReadNumVertices"))
+		.AddReturnType(EShaderFundamentalType::Uint);
+
+	OutFunctions.AddDefaulted_GetRef()
+		.SetName(TEXT("ReadAmount"))
+		.AddReturnType(EShaderFundamentalType::Float);
+
+	OutFunctions.AddDefaulted_GetRef()
+		.SetName(TEXT("ReadMask"))
+		.AddReturnType(EShaderFundamentalType::Float)
+		.AddParam(EShaderFundamentalType::Uint);
+
+	OutFunctions.AddDefaulted_GetRef()
+		.SetName(TEXT("ReadParents"))
+		.AddReturnType(EShaderFundamentalType::Int, 4) // int4
+		.AddParam(EShaderFundamentalType::Uint);
+
+	OutFunctions.AddDefaulted_GetRef()
+		.SetName(TEXT("ReadWeights"))
+		.AddReturnType(EShaderFundamentalType::Float, 4) // float4
+		.AddParam(EShaderFundamentalType::Uint);
+
+	OutFunctions.AddDefaulted_GetRef()
+		.SetName(TEXT("ReadOffset"))
+		.AddReturnType(EShaderFundamentalType::Float, 3) // float4
+		.AddParam(EShaderFundamentalType::Uint);
+
+	OutFunctions.AddDefaulted_GetRef()
+		.SetName(TEXT("GetEmbeddedPos"))
+		.AddReturnType(EShaderFundamentalType::Float, 3) // float3
+		.AddParam(EShaderFundamentalType::Uint);
+
+}
+
+BEGIN_SHADER_PARAMETER_STRUCT(FFleshDeformerDataInterfaceParameters, )
+	SHADER_PARAMETER(uint32, NumVertices)
+	SHADER_PARAMETER(uint32, BaseVertexIndex)
+	SHADER_PARAMETER(uint32, NumTetRestVertexBuffer)
+	SHADER_PARAMETER(uint32, NumTetVertexBuffer)
+	SHADER_PARAMETER(uint32, NumParentsBuffer)
+	SHADER_PARAMETER(uint32, NumWeightsBuffer)
+	SHADER_PARAMETER(uint32, NumOffsetBuffer)
+	SHADER_PARAMETER(uint32, NumMaskBuffer)
+	SHADER_PARAMETER(float, Amount)
+	SHADER_PARAMETER_RDG_BUFFER_SRV(Buffer<float>, TetRestVertexBuffer)
+	SHADER_PARAMETER_RDG_BUFFER_SRV(Buffer<float>, TetVertexBuffer)
+	SHADER_PARAMETER_RDG_BUFFER_SRV(Buffer<int>, ParentsBuffer)
+	SHADER_PARAMETER_RDG_BUFFER_SRV(Buffer<float>, WeightsBuffer)
+	SHADER_PARAMETER_RDG_BUFFER_SRV(Buffer<float>, OffsetBuffer)
+	SHADER_PARAMETER_RDG_BUFFER_SRV(Buffer<float>, MaskBuffer)
+END_SHADER_PARAMETER_STRUCT()
+
+void
+UDIFleshDeformer::GetShaderParameters(
+	TCHAR const* UID,
+	FShaderParametersMetadataBuilder& InOutBuilder,
+	FShaderParametersMetadataAllocations& InOutAllocations) const
+{
+	InOutBuilder.AddNestedStruct<FFleshDeformerDataInterfaceParameters>(UID);
+}
+
+void
+UDIFleshDeformer::GetPermutations(FComputeKernelPermutationVector& OutPermutationVector) const
+{
+	OutPermutationVector.AddPermutation(TEXT("ENABLE_DEFORMER_FLESH"), 2);
+}
+
+void
+UDIFleshDeformer::GetShaderHash(FString& InOutKey) const
+{
+	// Module registers "ChaosFlesh/Source/ChaosFleshEngine/Shaders/Private" as virtual shader 
+	// path: "/Plugin/ChaosFleshEngine".
+	GetShaderFileHash(
+		TEXT("/Plugin/ChaosFleshEngine/DIFleshDeformer.ush"),
+		EShaderPlatform::SP_PCD3D_SM5).AppendString(InOutKey);
+}
+
+void
+UDIFleshDeformer::GetHLSL(FString& OutHLSL, FString const& InDataInterfaceName) const
+{
+	TMap<FString, FStringFormatArg> TemplateArgs =
+	{
+		{ TEXT("DataInterfaceName"), InDataInterfaceName },
+	};
+
+	FString TemplateFile;
+	LoadShaderSourceFile(
+		TEXT("/Plugin/ChaosFleshEngine/DIFleshDeformer.ush"),
+		EShaderPlatform::SP_PCD3D_SM5, &TemplateFile, nullptr);
+	OutHLSL += FString::Format(*TemplateFile, TemplateArgs);
+}
+
+UComputeDataProvider*
+UDIFleshDeformer::CreateDataProvider(TObjectPtr<UObject> InBinding, uint64 InInputMask, uint64 InOutputMask) const
+{
+	UDIFleshDeformerDataProvider* Provider = NewObject<UDIFleshDeformerDataProvider>();
+
+	if ((Provider->SkinnedMesh = Cast<USkinnedMeshComponent>(InBinding)))
+	{
+		AActor* Actor = Provider->SkinnedMesh->GetOwner();
+		const TSet<UActorComponent*>& Components = Actor->GetComponents();
+		for (TSet<UActorComponent*>::TConstIterator It = Components.CreateConstIterator(); It; ++It)
+		{
+			if (const UActorComponent* Component = *It)
+			{
+				if ((Provider->FleshMesh = const_cast<UFleshComponent*>(Cast<UFleshComponent>(Component))))
+				{
+					// TODO: Validate that this FleshMesh has bindings for this SkeletalMesh
+					break;
+				}
+				// TODO: Someday, if/when Optimus supports deforming static meshes...
+				//if (Provider->StaticMesh = Cast<UStaticMeshComponent>(InBinding))
+				//{}
+			}
+		}
+	}
+
+	Provider->FleshDeformerParameters = FleshDeformerParameters;
+	return Provider;
+}
+
+//
+// DataProvider
+//
+
+bool
+UDIFleshDeformerDataProvider::IsValid() const
+{
+	return (SkinnedMesh && FleshMesh) || (StaticMesh && FleshMesh);
+}
+
+FComputeDataProviderRenderProxy*
+UDIFleshDeformerDataProvider::GetRenderProxy()
+{
+	return new FDIFleshDeformerProviderProxy(SkinnedMesh, FleshMesh, FleshDeformerParameters);
+}
+
+//
+// Proxy
+//
+
+FDIFleshDeformerProviderProxy::FDIFleshDeformerProviderProxy(
+	USkinnedMeshComponent* SkinnedMeshComponentIn, 
+	UFleshComponent* FleshComponentIn,
+	FFleshDeformerParameters const& FleshDeformerParametersIn)
+	: SkinnedMeshComponent(SkinnedMeshComponentIn)
+	, SkeletalMeshObject(SkinnedMeshComponentIn ? SkinnedMeshComponentIn->MeshObject : nullptr)
+	, FleshComponent(FleshComponentIn)
+	, FleshDeformerParameters(FleshDeformerParametersIn)
+{
+	LodIndex = SkeletalMeshObject ? SkeletalMeshObject->GetLOD() : 0;
+	if (SkinnedMeshComponent)
+	{
+		if (USkinnedAsset* SkinnedAsset = SkinnedMeshComponent->GetSkinnedAsset())
+		{
+			FPrimaryAssetId Id = SkinnedAsset->GetPrimaryAssetId();
+			MeshId = Id.IsValid() ? Id.ToString() : SkinnedAsset->GetName();
+		}
+	}
+}
+
+bool 
+FDIFleshDeformerProviderProxy::IsValid(FValidationData const& InValidationData) const
+{
+	if (InValidationData.ParameterStructSize != sizeof(FParameters))
+	{
+		return false;
+	}
+	if (SkinnedMeshComponent == nullptr || SkeletalMeshObject == nullptr || FleshComponent == nullptr)
+	{
+		return false;
+	}
+
+	const UFleshAsset* RestCollectionAsset = FleshComponent->GetRestCollection();
+	if (!RestCollectionAsset)
+	{
+		return false;
+	}
+
+	const FFleshCollection* FleshCollection = RestCollectionAsset->GetCollection();
+	if (!FleshCollection)
+	{
+		return false;
+	}
+
+	// Note: tet bindings will be invalid until initialized to a bindings group.
+	// We're just testing for the existence of such a group here.
+	GeometryCollection::Facades::FTetrahedralBindings TetBindings(*FleshCollection);
+	FName MeshName(*MeshId, MeshId.Len());
+	int32 TetMeshIdx = TetBindings.GetTetMeshIndex(MeshName, LodIndex);
+	if (TetMeshIdx == INDEX_NONE)
+	{
+		UE_LOG(LogFleshDeformer, Error, 
+			TEXT("FleshDeformer: Failed to find tet mesh index associated to mesh '%s' LOD: %d"), 
+			*MeshId, LodIndex);
+		return false;
+	}
+	bool bValid = TetBindings.ContainsBindingsGroup(TetMeshIdx, MeshName, LodIndex);
+	return bValid;
+}
+
+struct FFleshDeformerDataInterfacePermutationIds
+{
+	uint32 EnableDeformerFlesh = 0;
+
+	FFleshDeformerDataInterfacePermutationIds(FComputeKernelPermutationVector const& PermutationVector)
+	{
+		static FString Name(TEXT("ENABLE_DEFORMER_FLESH"));
+		static uint32 Hash = GetTypeHash(Name);
+		EnableDeformerFlesh = PermutationVector.GetPermutationBits(Name, Hash, 1);
+	}
+};
+
+void
+FDIFleshDeformerProviderProxy::GatherPermutations(FPermutationData& InOutPermutationData) const
+{
+	FFleshDeformerDataInterfacePermutationIds PermutationIds(InOutPermutationData.PermutationVector);
+	for (int32 InvocationIndex = 0; InvocationIndex < InOutPermutationData.NumInvocations; ++InvocationIndex)
+	{
+		InOutPermutationData.PermutationIds[InvocationIndex] |= PermutationIds.EnableDeformerFlesh;
+	}
+}
+
+void
+FDIFleshDeformerProviderProxy::AllocateResources(FRDGBuilder& GraphBuilder)
+{
+	//
+	// Rest Vertices from FleshComponent's rest collection.
+	//
+
+	check(FleshComponent);
+	const UFleshAsset* RestCollectionAsset = FleshComponent->GetRestCollection();
+	const FFleshCollection* FleshCollection = RestCollectionAsset->GetCollection();
+	NumTetRestVertices = 0;
+	NumTetVertices = 0;
+
+	const TManagedArray<FVector3f>* RestVertex = nullptr;
+	if (FleshCollection)
+	{
+		RestVertex = &FleshCollection->Vertex;
+		NumTetRestVertices = RestVertex->Num();
+	}
+
+	TetRestVertexBuffer = GraphBuilder.CreateBuffer(
+		FRDGBufferDesc::CreateStructuredDesc(sizeof(float), FGenericPlatformMath::Max((uint32)1, NumTetRestVertices * 3)),
+		TEXT("FleshDeformer.RestVertex"));
+	TetRestVertexBufferSRV = GraphBuilder.CreateSRV(TetRestVertexBuffer);
+
+	if (NumTetRestVertices)
+	{
+		const FVector3f* Vert0 = RestVertex->GetData();
+		const float* Vert00 = &Vert0->X;
+
+		GraphBuilder.QueueBufferUpload(
+			TetRestVertexBuffer,
+			Vert00,
+			NumTetRestVertices * 3 * sizeof(float),
+			ERDGInitialDataFlags::None);
+	}
+	else
+	{
+		GraphBuilder.QueueBufferUpload(
+			TetRestVertexBuffer,
+			&NullFloatBuffer,
+			1 * sizeof(float),
+			ERDGInitialDataFlags::None);
+	}
+
+	//
+	// Vertices from FleshComponent's dynamic collection, if they exist; rest
+	// collection vertices otherwise.
+	//
+
+	const TManagedArray<FVector3f>* Vertex = nullptr;
+	if (UFleshDynamicAsset* DynamicCollectionAsset = FleshComponent->GetDynamicCollection())
+	{
+		Vertex = &DynamicCollectionAsset->GetPositions();
+		NumTetVertices = Vertex->Num();
+	}
+	if (!NumTetVertices && FleshCollection)
+	{
+		Vertex = &FleshCollection->Vertex;
+		NumTetVertices = Vertex->Num();
+	}
+
+	TetVertexBuffer = GraphBuilder.CreateBuffer(
+		FRDGBufferDesc::CreateStructuredDesc(sizeof(float), FGenericPlatformMath::Max((uint32)1, NumTetVertices * 3)),
+		TEXT("FleshDeformer.Vertex"));
+	TetVertexBufferSRV = GraphBuilder.CreateSRV(TetVertexBuffer);
+
+	if (NumTetVertices)
+	{
+		const FVector3f* Vert0 = Vertex->GetData();
+		const float* Vert00 = &Vert0->X;
+
+		GraphBuilder.QueueBufferUpload(
+			TetVertexBuffer,
+			Vert00,
+			NumTetVertices * 3 * sizeof(float),
+			ERDGInitialDataFlags::None);
+	}
+	else
+	{
+		GraphBuilder.QueueBufferUpload(
+			TetVertexBuffer,
+			&NullFloatBuffer,
+			1 * sizeof(float),
+			ERDGInitialDataFlags::None);
+	}
+
+	//
+	// Bindings data from FleshComponent's rest collection.
+	//
+
+	GeometryCollection::Facades::FTetrahedralBindings* TetBindings = nullptr;
+	if (FleshCollection)
+	{
+		TetBindings = new GeometryCollection::Facades::FTetrahedralBindings(*FleshCollection);
+		FName MeshName(*MeshId, MeshId.Len());
+		TetIndex = TetBindings->GetTetMeshIndex(MeshName, LodIndex);
+		if (TetIndex != INDEX_NONE)
+		{
+			if (!TetBindings->ReadBindingsGroup(TetIndex, MeshName, LodIndex))
+			{
+				// Warn?
+			}
+		}
+		if (!TetBindings->IsValid())
+		{
+			delete TetBindings;
+			TetBindings = nullptr;
+		}
+	}
+	const TManagedArrayAccessor<FIntVector4>* Parents = TetBindings ? TetBindings->GetParentsRO() : nullptr;
+	const TManagedArrayAccessor<FVector4f>* Weights = TetBindings ? TetBindings->GetWeightsRO() : nullptr;
+	const TManagedArrayAccessor<FVector3f>* Offset = TetBindings ? TetBindings->GetOffsetsRO() : nullptr;
+	const TManagedArrayAccessor<float>* Mask = TetBindings ? TetBindings->GetMaskRO() : nullptr;
+
+	NumParents = (Parents ? Parents->Num() : 0);
+	ParentsBuffer =
+		GraphBuilder.CreateBuffer(
+			FRDGBufferDesc::CreateStructuredDesc(sizeof(int32), FGenericPlatformMath::Max((uint32)1, NumParents * 4)),
+			TEXT("FleshDeformer.Parents"));
+	ParentsBufferSRV = GraphBuilder.CreateSRV(ParentsBuffer);
+
+	NumWeights = (Weights ? Weights->Num() : 0);
+	WeightsBuffer =
+		GraphBuilder.CreateBuffer(
+			FRDGBufferDesc::CreateStructuredDesc(sizeof(float), FGenericPlatformMath::Max((uint32)1, NumWeights * 4)),
+			TEXT("FleshDeformer.Weights"));
+	WeightsBufferSRV = GraphBuilder.CreateSRV(WeightsBuffer);
+
+	NumOffset = (Offset ? Offset->Num() : 0);
+	OffsetBuffer =
+		GraphBuilder.CreateBuffer(
+			FRDGBufferDesc::CreateStructuredDesc(sizeof(float), FGenericPlatformMath::Max((uint32)1, NumOffset * 3)),
+			TEXT("FleshDeformer.Offset"));
+	OffsetBufferSRV = GraphBuilder.CreateSRV(OffsetBuffer);
+
+	NumMask = (Mask ? Mask->Num() : 0);
+	MaskBuffer =
+		GraphBuilder.CreateBuffer(
+			FRDGBufferDesc::CreateStructuredDesc(sizeof(float), FGenericPlatformMath::Max((uint32)1, NumMask)),
+			TEXT("FleshDeformer.Mask"));
+	MaskBufferSRV = GraphBuilder.CreateSRV(MaskBuffer);
+
+	if (Parents && Parents->Num())
+	{
+		const TManagedArray<FIntVector4>& ParentsValues = Parents->Get();
+		GraphBuilder.QueueBufferUpload(
+			ParentsBuffer,
+			&ParentsValues[0][0],
+			ParentsValues.Num() * 4 * sizeof(int32),
+			ERDGInitialDataFlags::None);
+	}
+	else
+	{
+		GraphBuilder.QueueBufferUpload(
+			ParentsBuffer,
+			&NullIntBuffer,
+			1 * sizeof(int32),
+			ERDGInitialDataFlags::None);
+	}
+	if (Weights && Weights->Num())
+	{
+		const TManagedArray<FVector4f>& WeightsValues = Weights->Get();
+		GraphBuilder.QueueBufferUpload(
+			WeightsBuffer,
+			&WeightsValues.GetData()->X,
+			WeightsValues.Num() * 4 * sizeof(float),
+			ERDGInitialDataFlags::None);
+	}
+	else
+	{
+		GraphBuilder.QueueBufferUpload(
+			WeightsBuffer,
+			&NullFloatBuffer,
+			1 * sizeof(float),
+			ERDGInitialDataFlags::None);
+	}
+	if (Offset && Offset->Num())
+	{
+		const TManagedArray<FVector3f>& OffsetValues = Offset->Get();
+		GraphBuilder.QueueBufferUpload(
+			OffsetBuffer,
+			&OffsetValues.GetData()->X,
+			OffsetValues.Num() * 3 * sizeof(float),
+			ERDGInitialDataFlags::None);
+	}
+	else
+	{
+		GraphBuilder.QueueBufferUpload(
+			OffsetBuffer,
+			&NullFloatBuffer,
+			1 * sizeof(float),
+			ERDGInitialDataFlags::None);
+	}
+	if (Mask && Mask->Num())
+	{
+		const TManagedArray<float>& MaskValues = Mask->Get();
+		GraphBuilder.QueueBufferUpload(
+			MaskBuffer,
+			&MaskValues[0],
+			MaskValues.Num() * sizeof(float),
+			ERDGInitialDataFlags::None);
+	}
+	else
+	{
+		GraphBuilder.QueueBufferUpload(
+			MaskBuffer,
+			&NullFloatBuffer,
+			1 * sizeof(float),
+			ERDGInitialDataFlags::None);
+	}
+}
+
+void
+FDIFleshDeformerProviderProxy::GatherDispatchData(FDispatchData const& InDispatchData)
+{
+	FSkeletalMeshRenderData const& SkeletalMeshRenderData = SkeletalMeshObject->GetSkeletalMeshRenderData();
+	FSkeletalMeshLODRenderData const* LodRenderData = &SkeletalMeshRenderData.LODRenderData[LodIndex];
+	const TStridedView<FParameters> ParameterArray = MakeStridedParameterView<FParameters>(InDispatchData);
+	for (int32 InvocationIndex = 0; InvocationIndex < ParameterArray.Num(); ++InvocationIndex)
+	{
+		FSkelMeshRenderSection const& RenderSection = LodRenderData->RenderSections[InvocationIndex];
+		
+		FParameters& Parameters = ParameterArray[InvocationIndex];
+
+		Parameters.NumVertices = RenderSection.GetNumVertices();
+		Parameters.BaseVertexIndex = RenderSection.BaseVertexIndex;
+
+		Parameters.Amount = FleshDeformerParameters.Amount;
+
+		Parameters.NumTetRestVertexBuffer = NumTetRestVertices;
+		Parameters.NumTetVertexBuffer = NumTetVertices;
+
+		Parameters.NumParentsBuffer = NumParents;
+		Parameters.NumWeightsBuffer = NumWeights;
+		Parameters.NumOffsetBuffer = NumOffset;
+		Parameters.NumMaskBuffer = NumMask;
+
+		Parameters.TetRestVertexBuffer = TetRestVertexBufferSRV;
+		Parameters.TetVertexBuffer = TetVertexBufferSRV;
+
+		Parameters.ParentsBuffer = ParentsBufferSRV;
+		Parameters.WeightsBuffer = WeightsBufferSRV;
+		Parameters.OffsetBuffer = OffsetBufferSRV;
+		Parameters.MaskBuffer = MaskBufferSRV;
+	}
+}
diff --git a/Engine/Plugins/Experimental/ChaosFlesh/Source/ChaosFleshEngine/Public/DataInterfaces/DIFleshDeformer.h b/Engine/Plugins/Experimental/ChaosFlesh/Source/ChaosFleshEngine/Public/DataInterfaces/DIFleshDeformer.h
new file mode 100644
--- /dev/null
+++ b/Engine/Plugins/Experimental/ChaosFlesh/Source/ChaosFleshEngine/Public/DataInterfaces/DIFleshDeformer.h
@@ -0,0 +1,136 @@
+// Copyright Epic Games, Inc. All Rights Reserved.
+#pragma once
+
+#include "OptimusComputeDataInterface.h"
+#include "ComputeFramework/ComputeDataProvider.h"
+#include "ChaosFlesh/FleshComponent.h"
+
+#include "DIFleshDeformer.generated.h"
+
+class FFleshDeformerDataInterfaceParameters;
+class FRDGBuffer;
+class FRDGBufferSRV;
+class FSkeletalMeshObject;
+class USkinnedMeshComponent;
+
+DECLARE_LOG_CATEGORY_EXTERN(LogFleshDeformer, Verbose, All);
+
+USTRUCT()
+struct FFleshDeformerParameters
+{
+	GENERATED_BODY();
+
+	/** Post-sim amount. */
+	UPROPERTY(EditAnywhere, Category = Deformer, meta = (ClampMin = "0.0", ClampMax = "1.0"))
+	float Amount = 1.0f;
+
+};
+
+
+
+/**
+* Compute Framework Data Interface for reading skeletal mesh and tetrahedral mesh bindings.
+*/
+UCLASS(Category = ComputeFramework)
+class /*OPTIMUSCORE_API*/ UDIFleshDeformer : public UOptimusComputeDataInterface
+{
+	GENERATED_BODY()
+
+public:
+	//~ Begin UOptimusComputeDataInterface Interface
+	FString GetDisplayName() const override;
+	TArray<FOptimusCDIPinDefinition> GetPinDefinitions() const override;
+	TSubclassOf<UActorComponent> GetRequiredComponentClass() const override;
+	//~ End UOptimusComputeDataInterface Interface
+
+	//~ Begin UComputeDataInterface Interface
+	TCHAR const* GetClassName() const override { return TEXT("FleshDeformer"); }
+	void GetSupportedInputs(TArray<FShaderFunctionDefinition>& OutFunctions) const override;
+	void GetShaderParameters(TCHAR const* UID, FShaderParametersMetadataBuilder& InOutBuilder, FShaderParametersMetadataAllocations& InOutAllocations) const override;
+	void GetPermutations(FComputeKernelPermutationVector& OutPermutationVector) const override;
+	void GetShaderHash(FString& InOutKey) const override;
+	void GetHLSL(FString& OutHLSL, FString const& InDataInterfaceName) const override;
+	UComputeDataProvider* CreateDataProvider(TObjectPtr<UObject> InBinding, uint64 InInputMask, uint64 InOutputMask) const override;
+	//~ End UComputeDataInterface Interface
+
+	UPROPERTY(EditAnywhere, Category = Deformer, meta = (ShowOnlyInnerProperties))
+	FFleshDeformerParameters FleshDeformerParameters;
+};
+
+/** Compute Framework Data Provider for reading skeletal mesh. */
+UCLASS(BlueprintType, editinlinenew, Category = ComputeFramework)
+class UDIFleshDeformerDataProvider : public UComputeDataProvider
+{
+	GENERATED_BODY()
+
+public:
+	UPROPERTY(EditAnywhere, BlueprintReadWrite, Category = Binding)
+	TObjectPtr<USkinnedMeshComponent> SkinnedMesh = nullptr;
+
+	UPROPERTY(EditAnywhere, BlueprintReadWrite, Category = Binding)
+	TObjectPtr<UStaticMeshComponent> StaticMesh = nullptr;
+
+	UPROPERTY(EditAnywhere, BlueprintReadWrite, Category = Binding)
+	TObjectPtr<UFleshComponent> FleshMesh = nullptr;
+
+	UPROPERTY()
+	FFleshDeformerParameters FleshDeformerParameters;
+
+	//~ Begin UComputeDataProvider Interface
+	bool IsValid() const override;
+	FComputeDataProviderRenderProxy* GetRenderProxy() override;
+	//~ End UComputeDataProvider Interface
+};
+
+class FDIFleshDeformerProviderProxy : public FComputeDataProviderRenderProxy
+{
+public:
+	FDIFleshDeformerProviderProxy(USkinnedMeshComponent* SkinnedMeshComponentIn, UFleshComponent* FleshComponentIn, const FFleshDeformerParameters& FleshDeformerParametersIn);
+
+	//~ Begin FComputeDataProviderRenderProxy Interface
+	bool IsValid(FValidationData const& InValidationData) const override;
+	void GatherPermutations(FPermutationData& InOutPermutationData) const override;
+	void AllocateResources(FRDGBuilder& GraphBuilder) override;
+	void GatherDispatchData(FDispatchData const& InDispatchData);
+	//~ End FComputeDataProviderRenderProxy Interface
+
+private:
+	using FParameters = FFleshDeformerDataInterfaceParameters;
+
+	USkinnedMeshComponent* SkinnedMeshComponent = nullptr;
+	FSkeletalMeshObject* SkeletalMeshObject = nullptr;
+	UFleshComponent* FleshComponent = nullptr;
+
+	FFleshDeformerParameters FleshDeformerParameters;
+
+	int32 TetIndex = INDEX_NONE;
+	int32 LodIndex = INDEX_NONE;
+	FString MeshId;
+
+	uint32 NumVertices = 0;
+
+	uint32 NumTetRestVertices = 0;
+	uint32 NumTetVertices = 0;
+
+	uint32 NumParents = 0;
+	uint32 NumWeights = 0;
+	uint32 NumOffset = 0;
+	uint32 NumMask = 0;
+
+	FRDGBuffer* TetRestVertexBuffer = nullptr;
+	FRDGBufferSRV* TetRestVertexBufferSRV = nullptr;
+	FRDGBuffer* TetVertexBuffer = nullptr;
+	FRDGBufferSRV* TetVertexBufferSRV = nullptr;
+
+	FRDGBuffer* ParentsBuffer = nullptr;
+	FRDGBufferSRV* ParentsBufferSRV = nullptr;
+	FRDGBuffer* WeightsBuffer = nullptr;
+	FRDGBufferSRV* WeightsBufferSRV = nullptr;
+	FRDGBuffer* OffsetBuffer = nullptr;
+	FRDGBufferSRV* OffsetBufferSRV = nullptr;
+	FRDGBuffer* MaskBuffer = nullptr;
+	FRDGBufferSRV* MaskBufferSRV = nullptr;
+
+	float NullFloatBuffer = 0.0f;
+	int32 NullIntBuffer = 0;
+};
\ No newline at end of file
diff --git a/Engine/Plugins/Experimental/ChaosFlesh/Source/ChaosFleshEngine/Shaders/Private/DIFleshDeformer.ush b/Engine/Plugins/Experimental/ChaosFlesh/Source/ChaosFleshEngine/Shaders/Private/DIFleshDeformer.ush
new file mode 100644
--- /dev/null
+++ b/Engine/Plugins/Experimental/ChaosFlesh/Source/ChaosFleshEngine/Shaders/Private/DIFleshDeformer.ush
@@ -0,0 +1,228 @@
+// Copyright Epic Games, Inc.  All Rights Reserved.
+
+uint {DataInterfaceName}_NumVertices;
+uint {DataInterfaceName}_BaseVertexIndex; // InputStreamStart
+
+uint {DataInterfaceName}_NumTetVertexBuffer;
+uint {DataInterfaceName}_NumTetRestVertexBuffer;
+
+uint {DataInterfaceName}_NumParentsBuffer;
+uint {DataInterfaceName}_NumWeightsBuffer;
+uint {DataInterfaceName}_NumOffsetBuffer;
+uint {DataInterfaceName}_NumMaskBuffer;
+
+float {DataInterfaceName}_Amount;
+
+Buffer<float> {DataInterfaceName}_TetRestVertexBuffer;
+Buffer<float> {DataInterfaceName}_TetVertexBuffer;
+
+Buffer<int>	{DataInterfaceName}_ParentsBuffer;
+Buffer<float> {DataInterfaceName}_WeightsBuffer;
+Buffer<float> {DataInterfaceName}_OffsetBuffer;
+Buffer<float> {DataInterfaceName}_MaskBuffer;
+
+uint ReadNumVertices_{DataInterfaceName}(uint Index)
+{
+	return {DataInterfaceName}_NumVertices;
+}
+
+float ReadAmount_{DataInterfaceName}()
+{
+	return {DataInterfaceName}_Amount;
+}
+
+//
+// Tet mesh points, rest and dynamic 
+//
+
+float3 ReadRestTetVertex_{DataInterfaceName}(uint Index)
+{
+#if ENABLE_DEFORMER_FLESH
+	if(Index >= 0 && Index < {DataInterfaceName}_NumTetRestVertexBuffer)
+	{
+		return float3(
+			{DataInterfaceName}_TetRestVertexBuffer[Index * 3], 
+			{DataInterfaceName}_TetRestVertexBuffer[Index * 3 + 1], 
+			{DataInterfaceName}_TetRestVertexBuffer[Index * 3 + 2]);
+	}
+#endif
+	return float3(0,0,0);
+}
+
+float3 ReadTetVertex_{DataInterfaceName}(uint Index)
+{
+#if ENABLE_DEFORMER_FLESH
+	if(Index >= 0 && Index < {DataInterfaceName}_NumTetVertexBuffer)
+	{
+		return float3(
+			{DataInterfaceName}_TetVertexBuffer[Index * 3], 
+			{DataInterfaceName}_TetVertexBuffer[Index * 3 + 1], 
+			{DataInterfaceName}_TetVertexBuffer[Index * 3 + 2]);
+	}
+#endif
+	return float3(0,0,0);
+}
+
+//
+// Render mesh to tet bindings
+//
+
+int4 ReadParents_{DataInterfaceName}(uint VertexIndex)
+{
+#if ENABLE_DEFORMER_FLESH
+	if(VertexIndex < {DataInterfaceName}_NumParentsBuffer)
+	{
+		uint BufferIndex = VertexIndex + {DataInterfaceName}_BaseVertexIndex;
+		return int4(
+			{DataInterfaceName}_ParentsBuffer[BufferIndex * 4], 
+			{DataInterfaceName}_ParentsBuffer[BufferIndex * 4 + 1], 
+			{DataInterfaceName}_ParentsBuffer[BufferIndex * 4 + 2],
+			{DataInterfaceName}_ParentsBuffer[BufferIndex * 4 + 3]);
+	}
+#endif
+	return int4(-1,-1,-1,-1);
+}
+
+float4 ReadWeights_{DataInterfaceName}(uint VertexIndex)
+{
+#if ENABLE_DEFORMER_FLESH
+	if(VertexIndex < {DataInterfaceName}_NumWeightsBuffer)
+	{
+		uint BufferIndex = VertexIndex + {DataInterfaceName}_BaseVertexIndex;
+		return float4(
+			{DataInterfaceName}_WeightsBuffer[BufferIndex * 4], 
+			{DataInterfaceName}_WeightsBuffer[BufferIndex * 4 + 1], 
+			{DataInterfaceName}_WeightsBuffer[BufferIndex * 4 + 2],
+			{DataInterfaceName}_WeightsBuffer[BufferIndex * 4 + 3]);
+	}
+#endif
+	return float4(0,0,0,0);
+}
+
+float3 ReadOffset_{DataInterfaceName}(uint VertexIndex)
+{
+#if ENABLE_DEFORMER_FLESH
+	if(VertexIndex < {DataInterfaceName}_NumOffsetBuffer)
+	{
+		uint BufferIndex = VertexIndex + {DataInterfaceName}_BaseVertexIndex;
+		return float3(
+			{DataInterfaceName}_OffsetBuffer[BufferIndex * 3], 
+			{DataInterfaceName}_OffsetBuffer[BufferIndex * 3 + 1], 
+			{DataInterfaceName}_OffsetBuffer[BufferIndex * 3 + 2]);
+	}
+#endif
+	return float3(0,0,0);
+}
+
+float ReadMask_{DataInterfaceName}(uint VertexIndex)
+{
+#if ENABLE_DEFORMER_FLESH
+	if(VertexIndex < {DataInterfaceName}_NumMaskBuffer)
+	{
+		uint BufferIndex = VertexIndex + {DataInterfaceName}_BaseVertexIndex;
+		return {DataInterfaceName}_MaskBuffer[BufferIndex];
+	}
+#endif
+	return 0.0;
+}
+
+//
+// Triangle orientation
+//
+
+
+// [ Duff et al. 2017, "Building an Orthonormal Basis, Revisited" ]
+// Discontinuity at TangentZ.z == 0
+float3x3 GetTangentBasis( float3 TangentZ )
+{
+    const float Sign = TangentZ.z >= 0 ? 1 : -1;
+    const float a = -rcp( Sign + TangentZ.z );
+    const float b = TangentZ.x * TangentZ.y * a;
+
+    float3 TangentX = { 1 + Sign * a * Pow2( TangentZ.x ), Sign * b, -Sign * TangentZ.x };
+    float3 TangentY = { b,  Sign + a * Pow2( TangentZ.y ), -TangentZ.y };
+
+    return float3x3( TangentX, TangentY, TangentZ );
+}
+
+float3 SafeNormalize(float3 Vec)
+{
+	float L = length(Vec);
+	if(L > 1.0e-6)
+	{
+		return Vec / L;
+	}
+	return float3(1,0,0);
+}
+
+float3x3 GetOrthogonalBasisVectors(float3 PtA, float3 PtB, float3 PtC)
+{
+	float3 EdgeBA = PtB - PtA;
+	float3 EdgeCA = PtC - PtA;
+	float3 OrthoNorm = SafeNormalize(cross(EdgeBA, EdgeCA));
+	return GetTangentBasis(OrthoNorm);
+/*
+	float3 EdgeBA = SafeNormalize(PtB - PtA);
+	float3 EdgeCA = SafeNormalize(PtC - PtA);
+	float3 OrthoNorm = cross(EdgeBA, EdgeCA);
+	float3 OrthoCA = cross(EdgeBA, OrthoNorm);
+	return float3x3(
+		EdgeBA[0], OrthoCA[0], OrthoNorm[0],
+		EdgeBA[1], OrthoCA[1], OrthoNorm[1],
+		EdgeBA[2], OrthoCA[2], OrthoNorm[2]); // Col major
+*/
+}
+
+
+float3x3 GetRestOrthogonalBasisVectors(int4 Tet) // assumes Tet[3] == -1
+{
+	float3 PtA = ReadRestTetVertex_{DataInterfaceName}(Tet[0]);
+	float3 PtB = ReadRestTetVertex_{DataInterfaceName}(Tet[1]);
+	float3 PtC = ReadRestTetVertex_{DataInterfaceName}(Tet[2]);
+	return GetOrthogonalBasisVectors(PtA, PtB, PtC);
+}
+
+float3x3 GetDynamicOrthogonalBasisVectors(int4 Tet) // assumes Tet[3] == -1
+{
+	float3 PtA = ReadTetVertex_{DataInterfaceName}(Tet[0]);
+	float3 PtB = ReadTetVertex_{DataInterfaceName}(Tet[1]);
+	float3 PtC = ReadTetVertex_{DataInterfaceName}(Tet[2]);
+	return GetOrthogonalBasisVectors(PtA, PtB, PtC);
+}
+
+float3 GetRotatedOffsetVector(uint VertexIndex, int4 Tet)
+{
+	float3 Offset = { 0, 0, 0 };
+#if ENABLE_DEFORMER_FLESH
+	Offset = ReadOffset_{DataInterfaceName}(VertexIndex);
+	if(Offset[0] != 0 && Offset[1] != 0 && Offset[2] != 0)
+	{
+		float3x3 RestRotInv = transpose(GetRestOrthogonalBasisVectors(Tet)); // Ortho matrix, so T == I
+		float3x3 CurrRot = GetDynamicOrthogonalBasisVectors(Tet);
+		Offset = mul(Offset, mul(RestRotInv, CurrRot));
+	}
+#endif
+	return Offset;
+}
+
+//
+// Helpers
+//
+
+float3 GetEmbeddedPos_{DataInterfaceName}(uint VertexIndex)
+{
+	float3 Pos = { 0, 0, 0 };
+#if ENABLE_DEFORMER_FLESH
+	int4 Parents = ReadParents_{DataInterfaceName}(VertexIndex);
+	float4 Weights = ReadWeights_{DataInterfaceName}(VertexIndex);
+	UNROLL for(uint i=0; i < 4; i++)
+	{
+		if(Parents[i] >= 0 && Weights[i] != 0.0)
+		{
+			Pos += ReadTetVertex_{DataInterfaceName}(Parents[i]) * Weights[i];
+		}
+	}
+	Pos += GetRotatedOffsetVector(VertexIndex, Parents);
+#endif
+	return Pos;
+}
