Component and Actor updates.
-	Moved references to the Solver Actor out of the component.
-	Use the component references to access the solver from the component level.
-	Set Property categories to “Chaos Deformable”
-	Removed Pre/Post edit attribute changes at the component level.
-	Added component based enable simulation UFunction for blueprint owned simulations.

diff --git a/Engine/Plugins/Experimental/ChaosFlesh/Source/ChaosFleshEngine/Private/ChaosFlesh/ChaosDeformablePhysicsComponent.cpp b/Engine/Plugins/Experimental/ChaosFlesh/Source/ChaosFleshEngine/Private/ChaosFlesh/ChaosDeformablePhysicsComponent.cpp
--- a/Engine/Plugins/Experimental/ChaosFlesh/Source/ChaosFleshEngine/Private/ChaosFlesh/ChaosDeformablePhysicsComponent.cpp
+++ b/Engine/Plugins/Experimental/ChaosFlesh/Source/ChaosFleshEngine/Private/ChaosFlesh/ChaosDeformablePhysicsComponent.cpp
@@ -13,15 +13,12 @@ void UDeformablePhysicsComponent::OnCreatePhysicsState()
 {
 	Super::OnCreatePhysicsState();
 	UE_LOG(LogDeformablePhysicsComponentInternal, Log, TEXT("UDeformablePhysicsComponent::OnCreatePhysicsState()"));
-	if (PrimarySolver)
+	if (PrimarySolverComponent)
 	{
-		if (UDeformableSolverComponent* SolverComponent = PrimarySolver->GetDeformableSolverComponent())
+		FDeformableSolver::FGameThreadAccess GameThreadSolver = PrimarySolverComponent->GameThreadAccess();
+		if (GameThreadSolver())
 		{
-			FDeformableSolver::FGameThreadAccess GameThreadSolver = SolverComponent->GameThreadAccess();
-			if (GameThreadSolver())
-			{
-				AddProxy(GameThreadSolver);
-			}
+			AddProxy(GameThreadSolver);
 		}
 	}
 }
@@ -30,15 +27,12 @@ void UDeformablePhysicsComponent::OnDestroyPhysicsState()
 {
 	Super::OnDestroyPhysicsState();
 	UE_LOG(LogDeformablePhysicsComponentInternal, Log, TEXT("UDeformablePhysicsComponent::OnDestroyPhysicsState()"));
-	if (PrimarySolver)
+	if (PrimarySolverComponent)
 	{
-		if (UDeformableSolverComponent* SolverComponent = PrimarySolver->GetDeformableSolverComponent())
+		FDeformableSolver::FGameThreadAccess GameThreadSolver = PrimarySolverComponent->GameThreadAccess();
+		if (GameThreadSolver())
 		{
-			FDeformableSolver::FGameThreadAccess GameThreadSolver = SolverComponent->GameThreadAccess();
-			if (GameThreadSolver())
-			{
-				RemoveProxy(GameThreadSolver);
-			}
+			RemoveProxy(GameThreadSolver);
 		}
 	}
 }
@@ -52,71 +46,6 @@ bool UDeformablePhysicsComponent::HasValidPhysicsState() const
 	return PhysicsProxy != nullptr;
 }
 
-
-#if WITH_EDITOR
-void UDeformablePhysicsComponent::PreEditChange(FProperty* PropertyThatWillChange)
-{
-	Super::PreEditChange(PropertyThatWillChange);
-	if (PropertyThatWillChange && PropertyThatWillChange->GetFName() == GET_MEMBER_NAME_CHECKED(UFleshComponent, PrimarySolver))
-	{
-		PreEditChangePrimarySolver = PrimarySolver;
-	}
-}
-
-
-void UDeformablePhysicsComponent::PostEditChangeProperty(struct FPropertyChangedEvent& PropertyChangedEvent)
-{
-	Super::PostEditChangeProperty(PropertyChangedEvent);
-
-	//
-	// The UDeformablePhysicsComponent and the UDeformableSolverComponent hold references to each other. 
-	// If one of the attributes change, then the attribute on the other component needs to be updated. 
-	//
-	if (PropertyChangedEvent.Property && PropertyChangedEvent.Property->GetFName() == GET_MEMBER_NAME_CHECKED(UFleshComponent, PrimarySolver))
-	{
-		if (AActor* Owner = Cast<AActor>(GetOwner()))
-		{
-			if (PrimarySolver)
-			{
-				if (UDeformableSolverComponent* SolverComponent = PrimarySolver->GetDeformableSolverComponent())
-				{
-					if (!SolverComponent->DeformableActors.Contains(Owner))
-					{
-						SolverComponent->DeformableActors.Add(TObjectPtr<AActor>(Owner));
-					}
-				}
-			}
-			else if (PreEditChangePrimarySolver)
-			{
-				bool DrivesOtherComponentsOnActor = false;
-				TArray<UDeformablePhysicsComponent*> DeformableComponentsOnActor;
-				Owner->GetComponents<UDeformablePhysicsComponent>(DeformableComponentsOnActor);
-				for (UDeformablePhysicsComponent* DeformableComponent : DeformableComponentsOnActor)
-				{
-					if (DeformableComponent->PrimarySolver == PreEditChangePrimarySolver)
-					{
-						DrivesOtherComponentsOnActor = true;
-						break;
-					}
-				}
-				if (!DrivesOtherComponentsOnActor)
-				{
-					if (UDeformableSolverComponent* SolverComponent = PreEditChangePrimarySolver->GetDeformableSolverComponent())
-					{
-						if (SolverComponent->DeformableActors.Contains(Owner))
-						{
-							SolverComponent->DeformableActors.Remove(Owner);
-						}
-					}
-				}
-			}
-		}
-		PreEditChangePrimarySolver = nullptr;
-	}
-}
-#endif
-
-
 void UDeformablePhysicsComponent::AddProxy(Chaos::Softs::FDeformableSolver::FGameThreadAccess& GameThreadSolver)
 {
 	PhysicsProxy = NewProxy();
@@ -138,11 +67,11 @@ void UDeformablePhysicsComponent::RemoveProxy(Chaos::Softs::FDeformableSolver::F
 
 UDeformableSolverComponent* UDeformablePhysicsComponent::GetDeformableSolver()
 {
-	return PrimarySolver ? PrimarySolver->GetDeformableSolverComponent() : nullptr;
+	return PrimarySolverComponent;
 }
 const UDeformableSolverComponent* UDeformablePhysicsComponent::GetDeformableSolver() const
 { 
-	return PrimarySolver ? PrimarySolver->GetDeformableSolverComponent() : nullptr; 
+	return PrimarySolverComponent;
 }
 
 UDeformablePhysicsComponent::UDeformablePhysicsComponent(const FObjectInitializer& ObjectInitializer)
@@ -151,22 +80,16 @@ UDeformablePhysicsComponent::UDeformablePhysicsComponent(const FObjectInitialize
 }
 
 
-void UDeformablePhysicsComponent::EnableSimulation(ADeformableSolverActor* InActor)
+void UDeformablePhysicsComponent::EnableSimulation(UDeformableSolverComponent* DeformableSolverComponent)
 {
-	if (InActor)
+	if (DeformableSolverComponent)
 	{
-		if (UDeformableSolverComponent* SolverComponent = InActor->GetDeformableSolverComponent())
+		PrimarySolverComponent = DeformableSolverComponent;
+		if (!DeformableSolverComponent->DeformableComponents.Contains(this))
 		{
-			if (AActor* Owner = Cast<AActor>(this->GetOwner()))
-			{
-				PrimarySolver = InActor;
-				if (!SolverComponent->DeformableActors.Contains(Owner))
-				{
-					SolverComponent->DeformableActors.Add(Owner);
-				}
-				SolverComponent->AddDeformableProxy(this);
-			}
+			DeformableSolverComponent->DeformableComponents.Add(this);
 		}
+		DeformableSolverComponent->AddDeformableProxy(this);
 	}
 }
 
diff --git a/Engine/Plugins/Experimental/ChaosFlesh/Source/ChaosFleshEngine/Private/ChaosFlesh/ChaosDeformableSolverComponent.cpp b/Engine/Plugins/Experimental/ChaosFlesh/Source/ChaosFleshEngine/Private/ChaosFlesh/ChaosDeformableSolverComponent.cpp
--- a/Engine/Plugins/Experimental/ChaosFlesh/Source/ChaosFleshEngine/Private/ChaosFlesh/ChaosDeformableSolverComponent.cpp
+++ b/Engine/Plugins/Experimental/ChaosFlesh/Source/ChaosFleshEngine/Private/ChaosFlesh/ChaosDeformableSolverComponent.cpp
@@ -76,68 +76,6 @@ UDeformableSolverComponent::GameThreadAccess()
 }
 
 
-#if WITH_EDITOR
-void UDeformableSolverComponent::PreEditChange(FProperty* PropertyThatWillChange)
-{
-	Super::PreEditChange(PropertyThatWillChange);
-	if (PropertyThatWillChange && PropertyThatWillChange->GetFName() == GET_MEMBER_NAME_CHECKED(UDeformableSolverComponent, DeformableActors))
-	{
-		PreEditChangeDeformableActors = DeformableActors;
-	}
-}
-
-void UDeformableSolverComponent::PostEditChangeProperty(struct FPropertyChangedEvent& PropertyChangedEvent)
-{
-	Super::PostEditChangeProperty(PropertyChangedEvent);
-
-	//
-	// The UDeformablePhysicsComponent and the UDeformableSolverComponent hold references to each other. 
-	// If one of the attributes change, then the attribute on the other component needs to be updated. 
-	//
-	if (PropertyChangedEvent.Property && PropertyChangedEvent.Property->GetFName() == GET_MEMBER_NAME_CHECKED(UDeformableSolverComponent, DeformableActors))
-	{
-		// process removed actors
-		for (TObjectPtr<AActor>& SimulatedActor : PreEditChangeDeformableActors)
-		{
-			if (SimulatedActor)
-			{
-				if (!DeformableActors.Contains(SimulatedActor))
-				{
-					TArray<UDeformablePhysicsComponent*> DeformableComponentsOnActor;
-					SimulatedActor->GetComponents<UDeformablePhysicsComponent>(DeformableComponentsOnActor);
-					for (UDeformablePhysicsComponent* DeformableComponent : DeformableComponentsOnActor)
-					{
-						if (DeformableComponent->PrimarySolver == GetOwner())
-						{
-							DeformableComponent->PrimarySolver = nullptr;
-						}
-					}
-				}
-			}
-		}
-		// process added actors
-		for (TObjectPtr<AActor>& SimulatedActor : DeformableActors)
-		{
-			if (SimulatedActor)
-			{
-				if (!PreEditChangeDeformableActors.Contains(SimulatedActor))
-				{
-					TArray<UDeformablePhysicsComponent*> DeformableComponentsOnActor;
-					SimulatedActor->GetComponents<UDeformablePhysicsComponent>(DeformableComponentsOnActor);
-					for (UDeformablePhysicsComponent* DeformableComponent : DeformableComponentsOnActor)
-					{
-						if (DeformableComponent->PrimarySolver == nullptr)
-						{
-							DeformableComponent->PrimarySolver = Cast<ADeformableSolverActor>(GetOwner());
-						}
-					}
-				}
-			}
-		}
-	}
-}
-#endif
-
 bool UDeformableSolverComponent::IsSimulatable() const
 {
 	return true;
@@ -147,8 +85,8 @@ bool UDeformableSolverComponent::IsSimulating(UDeformablePhysicsComponent* InCom
 {
 	if (InComponent)
 	{
-		const AActor* ComponentSolver = InComponent->PrimarySolver.Get();
-		return ComponentSolver == GetOwner();
+		const UDeformableSolverComponent* ComponentSolver = InComponent->PrimarySolverComponent.Get();
+		return ComponentSolver ==this;
 	}
 	return false;
 }
@@ -252,50 +190,19 @@ void UDeformableSolverComponent::Reset()
 			, bEnableGravity
 		}));
 
-		for (TObjectPtr<AActor>& DeformableActor : DeformableActors)
+		for (TObjectPtr<UDeformablePhysicsComponent>& DeformableComponent : DeformableComponents)
 		{
-			if (DeformableActor)
+			if( DeformableComponent )
 			{
-				TArray<UDeformablePhysicsComponent*> DeformableComponents;
-				DeformableActor->GetComponents<UDeformablePhysicsComponent>(DeformableComponents);
-
-				for (UDeformablePhysicsComponent* DeformableComponent : DeformableComponents)
+				if (IsSimulating(DeformableComponent))
 				{
-					if (IsSimulating(DeformableComponent))
-					{
-						AddDeformableProxy(DeformableComponent);
-					}
+					AddDeformableProxy(DeformableComponent);
 				}
 			}
 		}
 	}
 }
 
-/*
-* @todo(flesh) : In game additions. 
-void UDeformableSolverComponent::AddDeformableActor(TObjectPtr<AActor> InActor)
-{
-	if (InActor)
-	{
-		if (!DeformableActors.Contains(InActor))
-		{
-			DeformableActors.Add(InActor);
-		}
-
-		if (AFleshActor* FleshActor = Cast<AFleshActor>(InActor.Get()))
-		{
-			TArray<UDeformablePhysicsComponent*> DeformableComponents;
-			InActor->GetComponents<UDeformablePhysicsComponent>(DeformableComponents);
-
-			for (UDeformablePhysicsComponent* DeformableComponent : DeformableComponents)
-			{
-				AddDeformableProxy(DeformableComponent);
-			}
-		}
-	}
-}
-*/
-
 void UDeformableSolverComponent::AddDeformableProxy(UDeformablePhysicsComponent* InComponent)
 {
 	TRACE_CPUPROFILER_EVENT_SCOPE(DeformableSolverComponent_AddDeformableProxy);
@@ -331,21 +238,15 @@ void UDeformableSolverComponent::UpdateFromGameThread(float DeltaTime)
 		FDeformableSolver::FGameThreadAccess GameThreadSolver(Solver.Get(), Chaos::Softs::FGameThreadAccessor());
 
 		Chaos::Softs::FDeformableDataMap DataMap;
-		for (TObjectPtr<AActor>& DeformableActor : DeformableActors)
+		for (TObjectPtr<UDeformablePhysicsComponent>& DeformableComponent : DeformableComponents)
 		{
-			if (DeformableActor)
+			if (DeformableComponent)
 			{
-				TArray<UDeformablePhysicsComponent*> DeformableComponents;
-				DeformableActor->GetComponents<UDeformablePhysicsComponent>(DeformableComponents);
-
-				for (UDeformablePhysicsComponent* DeformableComponent : DeformableComponents)
+				if (IsSimulating(DeformableComponent))
 				{
-					if (IsSimulating(DeformableComponent))
+					if (FDataMapValue Value = DeformableComponent->NewDeformableData())
 					{
-						if (FDataMapValue Value = DeformableComponent->NewDeformableData())
-						{
-							DataMap.Add(DeformableComponent, Value);
-						}
+						DataMap.Add(DeformableComponent, Value);
 					}
 				}
 			}
@@ -370,21 +271,15 @@ void UDeformableSolverComponent::UpdateFromSimulation(float DeltaTime)
 
 		if (Output)
 		{
-			for (TObjectPtr<AActor>& DeformableActor : DeformableActors)
+			for (TObjectPtr<UDeformablePhysicsComponent>& DeformableComponent : DeformableComponents)
 			{
-				if (DeformableActor)
+				if (DeformableComponent)
 				{
-					TArray<UDeformablePhysicsComponent*> DeformableComponents;
-					DeformableActor->GetComponents<UDeformablePhysicsComponent>(DeformableComponents);
-
-					for (UDeformablePhysicsComponent* DeformableComponent : DeformableComponents)
+					if (IsSimulating(DeformableComponent))
 					{
-						if (IsSimulating(DeformableComponent))
+						if (const FDataMapValue* Buffer = Output->ObjectMap.Find(DeformableComponent))
 						{
-							if (const FDataMapValue* Buffer = Output->ObjectMap.Find(DeformableComponent))
-							{
-								DeformableComponent->UpdateFromSimualtion(Buffer);
-							}
+							DeformableComponent->UpdateFromSimualtion(Buffer);
 						}
 					}
 				}
diff --git a/Engine/Plugins/Experimental/ChaosFlesh/Source/ChaosFleshEngine/Private/ChaosFlesh/FleshActor.cpp b/Engine/Plugins/Experimental/ChaosFlesh/Source/ChaosFleshEngine/Private/ChaosFlesh/FleshActor.cpp
--- a/Engine/Plugins/Experimental/ChaosFlesh/Source/ChaosFleshEngine/Private/ChaosFlesh/FleshActor.cpp
+++ b/Engine/Plugins/Experimental/ChaosFlesh/Source/ChaosFleshEngine/Private/ChaosFlesh/FleshActor.cpp
@@ -2,6 +2,8 @@
 
 #include "ChaosFlesh/FleshActor.h"
 
+#include "ChaosFlesh/ChaosDeformableSolverComponent.h"
+#include "ChaosFlesh/ChaosDeformableSolverActor.h"
 
 //#include "ChaosFlesh/PB.h"
 
@@ -19,7 +21,71 @@ AFleshActor::AFleshActor(const FObjectInitializer& ObjectInitializer)
 	PrimaryActorTick.bCanEverTick = true;
 }
 
+void AFleshActor::EnableSimulation(ADeformableSolverActor* InActor)
+{
+	if (InActor)
+	{
+		if (FleshComponent && FleshComponent->GetRestCollection())
+		{
+			FleshComponent->EnableSimulation(InActor->GetDeformableSolverComponent());
+		}
+	}
+}
+
+
 #if WITH_EDITOR
+void AFleshActor::PreEditChange(FProperty* PropertyThatWillChange)
+{
+	Super::PreEditChange(PropertyThatWillChange);
+	if (PropertyThatWillChange && PropertyThatWillChange->GetFName() == GET_MEMBER_NAME_CHECKED(AFleshActor, PrimarySolver))
+	{
+		PreEditChangePrimarySolver = PrimarySolver;
+	}
+}
+
+
+void AFleshActor::PostEditChangeProperty(struct FPropertyChangedEvent& PropertyChangedEvent)
+{
+	Super::PostEditChangeProperty(PropertyChangedEvent);
+
+	//
+	// The UDeformablePhysicsComponent and the UDeformableSolverComponent hold references to each other. 
+	// If one of the attributes change, then the attribute on the other component needs to be updated. 
+	//
+	if (PropertyChangedEvent.Property && PropertyChangedEvent.Property->GetFName() == GET_MEMBER_NAME_CHECKED(AFleshActor, PrimarySolver))
+	{
+		if (PrimarySolver)
+		{
+			if (UDeformableSolverComponent* SolverComponent = PrimarySolver->GetDeformableSolverComponent())
+			{
+				if (FleshComponent)
+				{
+					FleshComponent->PrimarySolverComponent = SolverComponent;
+					if (!SolverComponent->DeformableComponents.Contains(FleshComponent))
+					{
+						SolverComponent->DeformableComponents.Add(TObjectPtr<UDeformablePhysicsComponent>(FleshComponent));
+					}
+				}
+			}
+		}
+		else if (PreEditChangePrimarySolver)
+		{
+			if (UDeformableSolverComponent* SolverComponent = PreEditChangePrimarySolver->GetDeformableSolverComponent())
+			{
+				if (FleshComponent)
+				{
+					FleshComponent->PrimarySolverComponent = nullptr;
+					if (SolverComponent->DeformableComponents.Contains(FleshComponent))
+					{
+						SolverComponent->DeformableComponents.Remove(FleshComponent);
+					}
+				}
+			}
+		}
+	}
+}
+
+
 bool AFleshActor::GetReferencedContentObjects(TArray<UObject*>& Objects) const
 {
 	Super::GetReferencedContentObjects(Objects);
@@ -32,14 +98,4 @@ bool AFleshActor::GetReferencedContentObjects(TArray<UObject*>& Objects) const
 }
 #endif
 
-void AFleshActor::EnableSimulation(ADeformableSolverActor* InActor)
-{
-	if (InActor)
-	{
-		if (FleshComponent && FleshComponent->GetRestCollection())
-		{
-			FleshComponent->EnableSimulation(InActor);
-		}
-	}
-}
 
diff --git a/Engine/Plugins/Experimental/ChaosFlesh/Source/ChaosFleshEngine/Private/ChaosFlesh/FleshComponent.cpp b/Engine/Plugins/Experimental/ChaosFlesh/Source/ChaosFleshEngine/Private/ChaosFlesh/FleshComponent.cpp
--- a/Engine/Plugins/Experimental/ChaosFlesh/Source/ChaosFleshEngine/Private/ChaosFlesh/FleshComponent.cpp
+++ b/Engine/Plugins/Experimental/ChaosFlesh/Source/ChaosFleshEngine/Private/ChaosFlesh/FleshComponent.cpp
@@ -62,12 +62,9 @@ void UFleshComponent::BeginPlay()
 			PrimaryComponentTick.AddPrerequisite(SkeletalMeshComponent, SkeletalMeshComponent->PrimaryComponentTick);
 		}
 	}
-	if (PrimarySolver)
+	if (PrimarySolverComponent)
 	{
-		if (UDeformableSolverComponent* DeformableSolverComponent = PrimarySolver->GetDeformableSolverComponent())
-		{
-			PrimaryComponentTick.AddPrerequisite(DeformableSolverComponent, DeformableSolverComponent->PrimaryComponentTick);
-		}
+		PrimaryComponentTick.AddPrerequisite(PrimarySolverComponent, PrimarySolverComponent->PrimaryComponentTick);
 	}
 }
 
diff --git a/Engine/Plugins/Experimental/ChaosFlesh/Source/ChaosFleshEngine/Public/ChaosFlesh/ChaosDeformablePhysicsComponent.h b/Engine/Plugins/Experimental/ChaosFlesh/Source/ChaosFleshEngine/Public/ChaosFlesh/ChaosDeformablePhysicsComponent.h
--- a/Engine/Plugins/Experimental/ChaosFlesh/Source/ChaosFleshEngine/Public/ChaosFlesh/ChaosDeformablePhysicsComponent.h
+++ b/Engine/Plugins/Experimental/ChaosFlesh/Source/ChaosFleshEngine/Public/ChaosFlesh/ChaosDeformablePhysicsComponent.h
@@ -28,8 +28,7 @@ public:
 	~UDeformablePhysicsComponent() {}
 
 	UFUNCTION(BlueprintCallable, Category = "Physics")
-	void EnableSimulation(ADeformableSolverActor* Actor);
-
+	void EnableSimulation(UDeformableSolverComponent* DeformableSolverComponent);
 
 	virtual FThreadingProxy* NewProxy() { return nullptr; }
 	virtual void AddProxy(Chaos::Softs::FDeformableSolver::FGameThreadAccess& GameThreadSolver);
@@ -43,20 +42,14 @@ public:
 	virtual bool ShouldCreatePhysicsState() const override;
 	virtual bool HasValidPhysicsState() const override;
 
-#if WITH_EDITOR
-	ADeformableSolverActor* PreEditChangePrimarySolver = nullptr;
-	virtual void PreEditChange(FProperty* PropertyThatWillChange) override;
-	virtual void PostEditChangeProperty(struct FPropertyChangedEvent& PropertyChangedEvent) override;
-#endif
-
 	/** PrimarySolver */
 	UDeformableSolverComponent* GetDeformableSolver();
 	const UDeformableSolverComponent* GetDeformableSolver() const;
 	
-	UPROPERTY(EditAnywhere, BlueprintReadWrite, Category = Attributes)
-	TObjectPtr<ADeformableSolverActor> PrimarySolver;
+	UPROPERTY(EditAnywhere, BlueprintReadWrite, Category = "Chaos Deformable")
+	TObjectPtr<UDeformableSolverComponent> PrimarySolverComponent;
 
-	UPROPERTY(EditAnywhere, BlueprintReadWrite, Category = Attributes)
+	UPROPERTY(EditAnywhere, BlueprintReadWrite, Category = "Chaos Deformable")
 	bool bTempEnableGravity = true;
 
 	const FThreadingProxy* GetPhysicsProxy() const { return PhysicsProxy; }
diff --git a/Engine/Plugins/Experimental/ChaosFlesh/Source/ChaosFleshEngine/Public/ChaosFlesh/ChaosDeformableSolverActor.h b/Engine/Plugins/Experimental/ChaosFlesh/Source/ChaosFleshEngine/Public/ChaosFlesh/ChaosDeformableSolverActor.h
--- a/Engine/Plugins/Experimental/ChaosFlesh/Source/ChaosFleshEngine/Public/ChaosFlesh/ChaosDeformableSolverActor.h
+++ b/Engine/Plugins/Experimental/ChaosFlesh/Source/ChaosFleshEngine/Public/ChaosFlesh/ChaosDeformableSolverActor.h
@@ -21,7 +21,7 @@ public:
 	UDeformableSolverComponent* GetDeformableSolverComponent() { return SolverComponent; }
 	const UDeformableSolverComponent* GetDeformableSolverComponent() const { return SolverComponent; }
 
-	UPROPERTY(Category = SolverActor, VisibleAnywhere, BlueprintReadOnly, meta = (AllowPrivateAccess = "true"))
+	UPROPERTY(Category = "Chaos Deformable", VisibleAnywhere, BlueprintReadOnly, meta = (AllowPrivateAccess = "true"))
 	TObjectPtr<UDeformableSolverComponent> SolverComponent = nullptr;
 
 };
\ No newline at end of file
diff --git a/Engine/Plugins/Experimental/ChaosFlesh/Source/ChaosFleshEngine/Public/ChaosFlesh/ChaosDeformableSolverComponent.h b/Engine/Plugins/Experimental/ChaosFlesh/Source/ChaosFleshEngine/Public/ChaosFlesh/ChaosDeformableSolverComponent.h
--- a/Engine/Plugins/Experimental/ChaosFlesh/Source/ChaosFleshEngine/Public/ChaosFlesh/ChaosDeformableSolverComponent.h
+++ b/Engine/Plugins/Experimental/ChaosFlesh/Source/ChaosFleshEngine/Public/ChaosFlesh/ChaosDeformableSolverComponent.h
@@ -17,7 +17,7 @@ class UDeformablePhysicsComponent;
 /**
 *	UDeformableSolverComponent
 */
-UCLASS()
+UCLASS(meta = (BlueprintSpawnableComponent))
 class CHAOSFLESHENGINE_API UDeformableSolverComponent : public USceneComponent
 {
 	GENERATED_UCLASS_BODY()
@@ -32,12 +32,6 @@ public:
 	~UDeformableSolverComponent();
 	void UpdateTickGroup();
 
-#if WITH_EDITOR
-	TArray< TObjectPtr<AActor> > PreEditChangeDeformableActors;
-	virtual void PreEditChange(FProperty* PropertyThatWillChange) override;
-	virtual void PostEditChangeProperty(struct FPropertyChangedEvent& PropertyChangedEvent) override;
-#endif
-
 	/* Solver API */
 	FDeformableSolver::FGameThreadAccess GameThreadAccess();
 
@@ -58,62 +52,62 @@ public:
 
 	/* Properties */
 
-	UPROPERTY(EditAnywhere, BlueprintReadOnly, Category = Chaos)
-		TArray< TObjectPtr<AActor> > DeformableActors;
+	UPROPERTY(EditAnywhere, BlueprintReadOnly, Category = "Chaos Deformable")
+	TArray< TObjectPtr<UDeformablePhysicsComponent> > DeformableComponents;
 
 	/** ObjectType defines how to initialize the rigid objects state, Kinematic, Sleeping, Dynamic. */
-	UPROPERTY(EditAnywhere, BlueprintReadOnly , Category = Chaos)
-		EDeformableExecutionModel ExecutionModel = EDeformableExecutionModel::Chaos_Deformable_DuringPhysics;
+	UPROPERTY(EditAnywhere, BlueprintReadOnly , Category = "Chaos Deformable")
+	EDeformableExecutionModel ExecutionModel = EDeformableExecutionModel::Chaos_Deformable_DuringPhysics;
 
-	UPROPERTY(EditAnywhere, Category = Chaos)
+	UPROPERTY(EditAnywhere, Category = "Chaos Deformable")
 		int32 NumSubSteps = 2;
 
-	UPROPERTY(EditAnywhere, Category = Chaos)
+	UPROPERTY(EditAnywhere, Category = "Chaos Deformable")
 		int32 NumSolverIterations = 5;
 
-	UPROPERTY(EditAnywhere, Category = Chaos)
+	UPROPERTY(EditAnywhere, Category = "Chaos Deformable")
 		bool FixTimeStep = false;
 
-	UPROPERTY(EditAnywhere, Category = Chaos)
+	UPROPERTY(EditAnywhere, Category = "Chaos Deformable")
 		float TimeStepSize = 0.05;
 
-	UPROPERTY(EditAnywhere, Category = Chaos)
+	UPROPERTY(EditAnywhere, Category = "Chaos Deformable")
 		bool CacheToFile = false;
 
-	UPROPERTY(EditAnywhere, Category = Chaos)
+	UPROPERTY(EditAnywhere, Category = "Chaos Deformable")
 		bool bEnableKinematics = true;
 
-	UPROPERTY(EditAnywhere, Category = Chaos)
+	UPROPERTY(EditAnywhere, Category = "Chaos Deformable")
 		bool bUseFloor = true;
 
-	UPROPERTY(EditAnywhere, Category = Chaos)
+	UPROPERTY(EditAnywhere, Category = "Chaos Deformable")
 		bool bDoSelfCollision = false;
 
-	UPROPERTY(EditAnywhere, Category = Chaos)
+	UPROPERTY(EditAnywhere, Category = "Chaos Deformable")
 		bool bDoThreadedAdvance = true;
 
-	UPROPERTY(EditAnywhere, Category = "Chaos|Expermental" )
+	UPROPERTY(EditAnywhere, Category = "Chaos|Deformable|Expermental" )
 		bool bUseGridBasedConstraints = false;
 
-	UPROPERTY(EditAnywhere, Category = Chaos)
+	UPROPERTY(EditAnywhere, Category = "Chaos Deformable")
 		float GridDx = 25.;
 
-	UPROPERTY(EditAnywhere, Category = Chaos)
+	UPROPERTY(EditAnywhere, Category = "Chaos Deformable")
 		bool bDoQuasistatics = false;
 
-	UPROPERTY(EditAnywhere, Category = Chaos)
+	UPROPERTY(EditAnywhere, Category = "Chaos Deformable")
 		float YoungModulus = 100000;
 
-	UPROPERTY(EditAnywhere, Category = Chaos)
+	UPROPERTY(EditAnywhere, Category = "Chaos Deformable")
 		bool bDoBlended = false;
 
-	UPROPERTY(EditAnywhere, Category = Chaos)
+	UPROPERTY(EditAnywhere, Category = "Chaos Deformable")
 		float BlendedZeta = 0;
 
-	UPROPERTY(EditAnywhere, Category = Chaos)
+	UPROPERTY(EditAnywhere, Category = "Chaos Deformable")
 		float Damping = 0;
 
-	UPROPERTY(EditAnywhere, Category = Chaos)
+	UPROPERTY(EditAnywhere, Category = "Chaos Deformable")
 		bool bEnableGravity = true;
 	//UPROPERTY(EditAnywhere, Category = Chaos)
 	//	bool bWaitForParallelDeformableTask = true;
diff --git a/Engine/Plugins/Experimental/ChaosFlesh/Source/ChaosFleshEngine/Public/ChaosFlesh/FleshActor.h b/Engine/Plugins/Experimental/ChaosFlesh/Source/ChaosFleshEngine/Public/ChaosFlesh/FleshActor.h
--- a/Engine/Plugins/Experimental/ChaosFlesh/Source/ChaosFleshEngine/Public/ChaosFlesh/FleshActor.h
+++ b/Engine/Plugins/Experimental/ChaosFlesh/Source/ChaosFleshEngine/Public/ChaosFlesh/FleshActor.h
@@ -7,6 +7,7 @@
 #include "UObject/ObjectMacros.h"
 #include "UObject/ObjectPtr.h"
 #include "ChaosFlesh/FleshComponent.h"
+#include "ChaosFlesh/ChaosDeformableSolverActor.h"
 
 #include "FleshActor.generated.h"
 
@@ -19,15 +20,20 @@ class CHAOSFLESHENGINE_API AFleshActor : public AActor
 
 public:
 
-	/* FleshComponent */
-	UPROPERTY(VisibleAnywhere, BlueprintReadOnly, Category = Destruction, meta = (ExposeFunctionCategories = "Components|Flesh", AllowPrivateAccess = "true"))
-		TObjectPtr<UFleshComponent> FleshComponent;
-	UFleshComponent* GetFleshComponent() const { return FleshComponent; }
-
 	UFUNCTION(BlueprintCallable, Category = "Physics")
 	void EnableSimulation(ADeformableSolverActor* Actor);
 
+	UPROPERTY(VisibleAnywhere, BlueprintReadOnly, Category = "Chaos Deformable")
+	TObjectPtr<UFleshComponent> FleshComponent;
+	UFleshComponent* GetFleshComponent() const { return FleshComponent; }
+
+	UPROPERTY(EditAnywhere, BlueprintReadWrite, Category = "Chaos Deformable")
+	TObjectPtr<ADeformableSolverActor> PrimarySolver;
+
 #if WITH_EDITOR
+	ADeformableSolverActor* PreEditChangePrimarySolver = nullptr;
+	virtual void PreEditChange(FProperty* PropertyThatWillChange) override;
+	virtual void PostEditChangeProperty(struct FPropertyChangedEvent& PropertyChangedEvent) override;
 	virtual bool GetReferencedContentObjects(TArray<UObject*>& Objects) const override;
 #endif
 };
diff --git a/Engine/Plugins/Experimental/ChaosFlesh/Source/ChaosFleshEngine/Public/ChaosFlesh/FleshAsset.h b/Engine/Plugins/Experimental/ChaosFlesh/Source/ChaosFleshEngine/Public/ChaosFlesh/FleshAsset.h
--- a/Engine/Plugins/Experimental/ChaosFlesh/Source/ChaosFleshEngine/Public/ChaosFlesh/FleshAsset.h
+++ b/Engine/Plugins/Experimental/ChaosFlesh/Source/ChaosFleshEngine/Public/ChaosFlesh/FleshAsset.h
@@ -92,10 +92,6 @@ public:
 
 	void Serialize(FArchive& Ar);
 
-
-	UPROPERTY(EditAnywhere, Category = "Physics")
-	int32 Substeps=1;
-
 	//
 	// Dataflow
 	//
