ChaosFlesh: Radial meshing functions, engine level

diff --git a/Engine/Plugins/Experimental/ChaosFlesh/Source/ChaosFlesh/Private/Meshing/ChaosFleshRadialMeshing.cpp b/Engine/Plugins/Experimental/ChaosFlesh/Source/ChaosFlesh/Private/Meshing/ChaosFleshRadialMeshing.cpp
new file mode 100644
--- /dev/null
+++ b/Engine/Plugins/Experimental/ChaosFlesh/Source/ChaosFlesh/Private/Meshing/ChaosFleshRadialMeshing.cpp
@@ -0,0 +1,242 @@
+// Copyright Epic Games, Inc. All Rights Reserved.
+
+#include "Meshing/ChaosFleshRadialMeshing.h"
+
+using FReal = FVector::FReal;
+void RadialTetMesh(const FReal InnerRadius, const FReal OuterRadius, const FReal Height, const int32 RadialSample, const int32 AngularSample, const int32 VerticalSample, TArray<FIntVector4>& TetElements, TArray<FVector>& TetVertices) 
+{
+	TArray<FVector> HexVertices;
+	TArray<int32> HexElements;
+	RadialHexMesh(InnerRadius, OuterRadius, Height, RadialSample, AngularSample, VerticalSample, HexElements, HexVertices);
+	RegularHexMesh2TetMesh(HexVertices, HexElements, TetVertices, TetElements);
+}
+
+void RadialHexMesh(const FReal InnerRadius, const FReal OuterRadius, const FReal Height, const int32 RadialSample, const int32 AngularSample, const int32 VerticalSample, TArray<int32>& HexElements, TArray<FVector>& HexVertices) 
+{
+	HexElements.SetNum(8 * (VerticalSample - 1) * (RadialSample - 1) * AngularSample);
+	HexVertices.SetNum(VerticalSample * RadialSample * AngularSample);
+	FReal dr = (OuterRadius - InnerRadius) / FReal(RadialSample - 1);
+	FReal dtheta = FReal(2) * PI / FReal(AngularSample);
+	FReal dz = Height / FReal(VerticalSample - 1);
+	for (int32 i = 0; i < VerticalSample; i++)
+	{
+		//start at top, work way down
+		FReal ZVal = Height * FReal(0.5) - i * dz;
+		for (int32 j = 0; j < AngularSample; j++)
+		{
+			FReal XVal = cos(j * dtheta);
+			FReal YVal = sin(j * dtheta);
+			for (int32 k = 0; k < RadialSample; k++)
+			{
+				FReal RVal = InnerRadius + k * dr;
+				HexVertices[AngularSample * RadialSample * i + RadialSample * j + k] = { RVal * XVal,RVal * YVal,ZVal };
+
+			}
+		}
+	}
+	//CCW bottom then top local ordering of hex nodes
+	for (int32 i = 0; i < (VerticalSample - 1); i++)
+	{
+		for (int32 j = 0; j < (AngularSample - 1); j++)
+		{
+			for (int32 k = 0; k < (RadialSample - 1); k++)
+			{
+				//lower level
+				HexElements[8 * (i * (RadialSample - 1) * AngularSample + j * (RadialSample - 1) + k)] = RadialSample * AngularSample * (i + 1) + RadialSample * j + k;
+				HexElements[8 * (i * (RadialSample - 1) * AngularSample + j * (RadialSample - 1) + k) + 1] = RadialSample * AngularSample * (i + 1) + RadialSample * j + k + 1;
+				HexElements[8 * (i * (RadialSample - 1) * AngularSample + j * (RadialSample - 1) + k) + 2] = RadialSample * AngularSample * (i + 1) + RadialSample * (j + 1) + k + 1;
+				HexElements[8 * (i * (RadialSample - 1) * AngularSample + j * (RadialSample - 1) + k) + 3] = RadialSample * AngularSample * (i + 1) + RadialSample * (j + 1) + k;
+				//upper level
+				HexElements[8 * (i * (RadialSample - 1) * AngularSample + j * (RadialSample - 1) + k) + 4] = RadialSample * AngularSample * i + RadialSample * j + k;
+				HexElements[8 * (i * (RadialSample - 1) * AngularSample + j * (RadialSample - 1) + k) + 5] = RadialSample * AngularSample * i + RadialSample * j + k + 1;
+				HexElements[8 * (i * (RadialSample - 1) * AngularSample + j * (RadialSample - 1) + k) + 6] = RadialSample * AngularSample * i + RadialSample * (j + 1) + k + 1;
+				HexElements[8 * (i * (RadialSample - 1) * AngularSample + j * (RadialSample - 1) + k) + 7] = RadialSample * AngularSample * i + RadialSample * (j + 1) + k;
+			}
+		}
+		//Gluing edges, wrap-around
+		for (int32 k = 0; k < (RadialSample - 1); k++)
+		{
+			//lower level
+			HexElements[8 * (i * (RadialSample - 1) * AngularSample + (AngularSample - 1) * (RadialSample - 1) + k)] = RadialSample * AngularSample * (i + 1) + RadialSample * (AngularSample - 1) + k;
+			HexElements[8 * (i * (RadialSample - 1) * AngularSample + (AngularSample - 1) * (RadialSample - 1) + k) + 1] = RadialSample * AngularSample * (i + 1) + RadialSample * (AngularSample - 1) + k + 1;
+			HexElements[8 * (i * (RadialSample - 1) * AngularSample + (AngularSample - 1) * (RadialSample - 1) + k) + 2] = RadialSample * AngularSample * (i + 1) + k + 1;
+			HexElements[8 * (i * (RadialSample - 1) * AngularSample + (AngularSample - 1) * (RadialSample - 1) + k) + 3] = RadialSample * AngularSample * (i + 1) + k;
+			//upper level
+			HexElements[8 * (i * (RadialSample - 1) * AngularSample + (AngularSample - 1) * (RadialSample - 1) + k) + 4] = RadialSample * AngularSample * i + RadialSample * (AngularSample - 1) + k;
+			HexElements[8 * (i * (RadialSample - 1) * AngularSample + (AngularSample - 1) * (RadialSample - 1) + k) + 5] = RadialSample * AngularSample * i + RadialSample * (AngularSample - 1) + k + 1;
+			HexElements[8 * (i * (RadialSample - 1) * AngularSample + (AngularSample - 1) * (RadialSample - 1) + k) + 6] = RadialSample * AngularSample * i + k + 1;
+			HexElements[8 * (i * (RadialSample - 1) * AngularSample + (AngularSample - 1) * (RadialSample - 1) + k) + 7] = RadialSample * AngularSample * i + k;
+		}
+	}
+}
+
+void RegularHexMesh2TetMesh(const TArray<FVector>& HexVertices, const TArray<int32>& HexElements, TArray<FVector>& TetVertices, TArray<FIntVector4>& TetElements)
+{
+	if (ensureMsgf(HexElements.Num() % 8 == 0, TEXT("The hex mesh is not divisible by 8."))) 
+	{
+		TArray<FIntVector4> FaceToNodes = { {int32(0), int32(1), int32(2), int32(3)},{int32(0), int32(1), int32(5), int32(4)},{int32(1), int32(2), int32(6), int32(5)},{int32(3), int32(2), int32(6), int32(7)},{int32(0), int32(3), int32(7), int32(4)},{int32(4), int32(5), int32(6), int32(7)} };
+		TArray<FUintVector2> CommonFaces;
+		int32 Offset = HexVertices.Num();
+		TetVertices = HexVertices;
+		TetElements.SetNum(0);
+		for (int32 i = 0; i < int32(HexElements.Num() / 8); ++i)
+		{
+			FVector Center = FVector(0);
+			for (int32 ie = 0; ie < 8; ie++)
+			{
+				for (int32 c = 0; c < 3; c++)
+				{
+					Center[c] += HexVertices[HexElements[8 * i + ie]][c] / FReal(8);
+				}
+			}
+			TetVertices.Emplace(Center);
+		}
+		/*});*/
+		int32 q0, f0, q1, h0, h1, Nexti;
+
+		ComputeHexMeshFaces(HexElements, CommonFaces);
+
+		for (int32 f = 0; f < CommonFaces.Num(); ++f)
+		{
+			q0 = int32(CommonFaces[f][0]) / 6;
+			f0 = int32(CommonFaces[f][0]) % 6;
+			h0 = Offset + q0;
+			// The same face is shared by two hexs
+			if (CommonFaces[f][1] > uint32(-1))
+			{
+				q1 = int32(CommonFaces[f][1]) / 6;
+				h1 = Offset + q1;
+			}
+			// The face is on the boundary
+			else
+			{
+				h1 = TetVertices.Num();
+				FVector FaceCenter = FVector(0);
+				for (int32 ie = 0; ie < 4; ie++)
+			{
+					for (int32 c = 0; c < 3; c++)
+					{
+						FaceCenter[c] += HexVertices[HexElements[8 * q0 + FaceToNodes[f0][ie]]][c] / FReal(4);
+					}
+				}
+				TetVertices.Emplace(FaceCenter);
+			}
+			for (int32 i = 0; i < 4; i++)
+			{
+				Nexti = (i + 1) % 4;
+				FIntVector4 NewElement;
+				NewElement[0] = h0;
+				NewElement[1] = h1;
+				if (f0 == 1 || f0 == 2 || f0 == 5)
+				{
+					NewElement[2] = HexElements[8 * q0 + FaceToNodes[f0][i]];
+					NewElement[3] = HexElements[8 * q0 + FaceToNodes[f0][Nexti]];
+				}
+				else
+				{
+					NewElement[2] = HexElements[8 * q0 + FaceToNodes[f0][Nexti]];
+					NewElement[3] = HexElements[8 * q0 + FaceToNodes[f0][i]];
+				}
+				TetElements.Emplace(NewElement);
+			}
+		}
+	}
+}
+
+void ComputeHexMeshFaces(const TArray<int32>& HexElements, TArray<FUintVector2>& CommonFaces)
+{
+	auto FaceGreaterThan = [](FIntVector4& i1, FIntVector4& i2)
+	{
+		bool GreaterThan = false;
+		if (i1[0] > i2[0])
+			GreaterThan = true;
+		else if (i1[0] == i2[0] && i1[1] > i2[1])
+			GreaterThan = true;
+		else if (i1[0] == i2[0] && i1[1] == i2[1] && i1[2] > i2[2])
+			GreaterThan = true;
+		else if (i1[0] == i2[0] && i1[1] == i2[1] && i1[2] == i2[2] && i1[3] > i2[3])
+			GreaterThan = true;
+		return GreaterThan;
+	};
+	auto FaceEqual = [](FIntVector4& i1, FIntVector4& i2)
+	{
+		return (i1[0] == i2[0] && i1[1] == i2[1] && i1[2] == i2[2] && i1[3] == i2[3]);
+	};
+	auto GreaterThan = [&FaceGreaterThan](const TArray<int32>& Faces, int32 e1, int32 e2)
+	{
+		FIntVector4 Face1 = { Faces[4 * e1 + 0],Faces[4 * e1 + 1],Faces[4 * e1 + 2],Faces[4 * e1 + 3] };
+		FIntVector4 Face2 = { Faces[4 * e2 + 0],Faces[4 * e2 + 1],Faces[4 * e2 + 2],Faces[4 * e2 + 3] };
+		return FaceGreaterThan(Face1, Face2);
+	};
+
+	auto Equal = [&FaceEqual](const TArray<int32>& Faces, const TArray<int32>& AllFaces, int32 i1, int32 i2)
+	{
+		int32 e1 = AllFaces[i1];
+		int32 e2 = AllFaces[i2];
+		FIntVector4 Face1 = { Faces[4 * e1 + 0],Faces[4 * e1 + 1],Faces[4 * e1 + 2],Faces[4 * e1 + 3] };
+		FIntVector4 Face2 = { Faces[4 * e2 + 0],Faces[4 * e2 + 1],Faces[4 * e2 + 2],Faces[4 * e2 + 3] };
+		return FaceEqual(Face1, Face2);
+	};
+	ComputeMeshFaces(HexElements, GreaterThan, Equal, CommonFaces, GenerateHexFaceMesh, 4);
+}
+
+template <typename Func1, typename Func2, typename Func3>
+void ComputeMeshFaces(const TArray<int32>& mesh, Func1 GreaterThan, Func2 Equal, TArray<FUintVector2>& CommonFaces, Func3 GenerateFace, int32 PointsPerFace)
+{
+	TArray<int32> Faces;
+	GenerateFace(mesh, Faces);
+	int32 FaceNumber = Faces.Num() / PointsPerFace;
+	TArray<int32> AllFaces;
+	AllFaces.SetNum(FaceNumber);
+	TArray<int32> Ranges;
+	Ranges.SetNum(FaceNumber + 1);
+	for (int32 f = 0; f < int32(FaceNumber); f++)
+	{
+		AllFaces[f] = f;
+		Ranges[f] = f;
+	}
+
+	AllFaces.Sort([&GreaterThan, &Faces](int32 FaceA, int32 FaceB) { return GreaterThan(Faces, FaceA, FaceB); });
+	TArray<int32> UniqueRanges({ 0 });
+
+	for (int32 i = 0; i < Ranges.Num() - 2; ++i)
+	{
+		if (!Equal(Faces, AllFaces, Ranges[i], Ranges[i + 1]))
+		{
+			UniqueRanges.Emplace(i + 1);
+		}
+	}
+
+	int32 TotalFaces = UniqueRanges.Num();
+	UniqueRanges.SetNum(TotalFaces + 1);
+	UniqueRanges[TotalFaces] = int32(FaceNumber);
+
+
+	CommonFaces.SetNum(TotalFaces);
+	for (int32 f = 0; f < int32(CommonFaces.Num()); ++f)
+	{
+		CommonFaces[f][0] = AllFaces[UniqueRanges[f]];
+		if (UniqueRanges[f + 1] - UniqueRanges[f] == 1)
+			CommonFaces[f][1] = -1;
+		else
+		{
+			CommonFaces[f][1] = AllFaces[UniqueRanges[f] + 1];
+		}
+	}
+}
+
+void GenerateHexFaceMesh(const TArray<int32>& HexElements, TArray<int32>& Faces)
+{
+	if (ensureMsgf(HexElements.Num() % 8 == 0, TEXT("The hex mesh is not divisible by 8.")))
+	{
+		TArray<int32> FaceTemp = { int32(0), int32(1), int32(2), int32(3),int32(0), int32(1), int32(5), int32(4),int32(1), int32(2), int32(6), int32(5),int32(3), int32(2), int32(6), int32(7),int32(0), int32(3), int32(7), int32(4),int32(4), int32(5), int32(6), int32(7) };
+		Faces.SetNum(HexElements.Num() / 8 * 24);
+		for (int32 i = 0; i < HexElements.Num() / 8; i++)
+		{
+			for (int32 j = 0; j < int32(24); j++)
+			{
+				Faces[24 * i + j] = HexElements[8 * i + FaceTemp[j]];
+			}
+		}
+	}
+}
\ No newline at end of file
diff --git a/Engine/Plugins/Experimental/ChaosFlesh/Source/ChaosFlesh/Public/Meshing/ChaosFleshRadialMeshing.h b/Engine/Plugins/Experimental/ChaosFlesh/Source/ChaosFlesh/Public/Meshing/ChaosFleshRadialMeshing.h
new file mode 100644
--- /dev/null
+++ b/Engine/Plugins/Experimental/ChaosFlesh/Source/ChaosFlesh/Public/Meshing/ChaosFleshRadialMeshing.h
@@ -0,0 +1,12 @@
+// Copyright Epic Games, Inc. All Rights Reserved.
+#pragma once
+#include "CoreMinimal.h"
+
+void CHAOSFLESH_API RadialTetMesh(const FVector::FReal InnerRadius, const FVector::FReal OuterRadius, const FVector::FReal Height, const int32 RadialSample, const int32 AngularSample, const int32 VerticalSample, TArray<FIntVector4>& TetElements, TArray<FVector>& TetVertices);
+void CHAOSFLESH_API RadialHexMesh(const FVector::FReal InnerRadius, const FVector::FReal OuterRadius, const FVector::FReal Height, const int32 RadialSample, const int32 AngularSample, const int32 VerticalSample, TArray<int32>& HexElements, TArray<FVector>& HexVertices);
+void CHAOSFLESH_API RegularHexMesh2TetMesh(const TArray<FVector>& HexVertices, const TArray<int32>& HexElements, TArray<FVector>& TetVertices, TArray<FIntVector4>& TetElements);
+void CHAOSFLESH_API ComputeHexMeshFaces(const TArray<int32>& HexElements, TArray<FUintVector2>& CommonFaces);
+
+template <typename Func1, typename Func2, typename Func3>
+void CHAOSFLESH_API ComputeMeshFaces(const TArray<int32>& mesh, Func1 GreaterThan, Func2 Equal, TArray<FUintVector2>& CommonFaces, Func3 GenerateFace, int32 PointsPerFace);
+void CHAOSFLESH_API GenerateHexFaceMesh(const TArray<int32>& HexElements, TArray<int32>& Faces);
\ No newline at end of file
diff --git a/Engine/Plugins/Experimental/ChaosFlesh/Source/ChaosFleshNodes/Private/Dataflow/ChaosFleshNodesPlugin.cpp b/Engine/Plugins/Experimental/ChaosFlesh/Source/ChaosFleshNodes/Private/Dataflow/ChaosFleshNodesPlugin.cpp
--- a/Engine/Plugins/Experimental/ChaosFlesh/Source/ChaosFleshNodes/Private/Dataflow/ChaosFleshNodesPlugin.cpp
+++ b/Engine/Plugins/Experimental/ChaosFlesh/Source/ChaosFleshNodes/Private/Dataflow/ChaosFleshNodesPlugin.cpp
@@ -8,11 +8,11 @@
 #include "Dataflow/ChaosFleshEngineAssetNodes.h"
 #include "Dataflow/ChaosFleshFiberDirectionInitializationNodes.h"
 #include "Dataflow/ChaosFleshKinematicInitializationNodes.h"
+#include "Dataflow/ChaosFleshRadialTetrahedronNodes.h"
 #include "Dataflow/ChaosFleshRenderInitializationNodes.h"
 #include "Dataflow/ChaosFleshTetrahedralNodes.h"
 #include "Modules/ModuleManager.h"
 
-
 #define LOCTEXT_NAMESPACE "ChaosFleshNodes"
 
 
@@ -25,6 +25,7 @@ void IChaosFleshNodesPlugin::StartupModule()
 	Dataflow::RegisterChaosFleshKinematicInitializationNodes();
 	Dataflow::ChaosFleshTetrahedralNodes();
 	Dataflow::ChaosFleshBindingsNodes();
+	Dataflow::ChaosFleshRadialTetrahedronNodes();
 }
 
 void IChaosFleshNodesPlugin::ShutdownModule()
diff --git a/Engine/Plugins/Experimental/ChaosFlesh/Source/ChaosFleshNodes/Private/Dataflow/ChaosFleshRadialTetrahedronNodes.cpp b/Engine/Plugins/Experimental/ChaosFlesh/Source/ChaosFleshNodes/Private/Dataflow/ChaosFleshRadialTetrahedronNodes.cpp
new file mode 100644
--- /dev/null
+++ b/Engine/Plugins/Experimental/ChaosFlesh/Source/ChaosFleshNodes/Private/Dataflow/ChaosFleshRadialTetrahedronNodes.cpp
@@ -0,0 +1,54 @@
+// Copyright Epic Games, Inc. All Rights Reserved.
+
+#include "Dataflow/ChaosFleshRadialTetrahedronNodes.h"
+#include "Meshing/ChaosFleshRadialMeshing.h"
+#include "Dataflow/ChaosFleshTetrahedralNodes.h"
+
+#include "Chaos/Deformable/Utilities.h"
+#include "ChaosFlesh/ChaosFlesh.h"
+#include "Chaos/Utilities.h"
+#include "Chaos/UniformGrid.h"
+#include "ChaosFlesh/FleshCollection.h"
+#include "ChaosFlesh/FleshCollectionUtility.h"
+#include "ChaosLog.h"
+#include "Dataflow/DataflowInputOutput.h"
+#include "DynamicMesh/DynamicMesh3.h"
+#include "DynamicMesh/DynamicMeshAABBTree3.h"
+#include "Engine/StaticMesh.h"
+#include "FTetWildWrapper.h"
+#include "Generate/IsosurfaceStuffing.h"
+#include "GeometryCollection/ManagedArrayCollection.h"
+#include "MeshDescriptionToDynamicMesh.h"
+#include "Spatial/FastWinding.h"
+#include "Spatial/MeshAABBTree3.h"
+
+namespace Dataflow
+{
+	void ChaosFleshRadialTetrahedronNodes()
+	{
+		DATAFLOW_NODE_REGISTER_CREATION_FACTORY(FRadialTetrahedronDataflowNodes);
+	}
+}
+
+
+//=============================================================================
+// FRadialTetrahedronDataflowNodes
+//=============================================================================
+
+
+void FRadialTetrahedronDataflowNodes::Evaluate(Dataflow::FContext& Context, const FDataflowOutput* Out) const
+{
+	if (Out->IsA<DataType>(&Collection))
+	{
+		TUniquePtr<FFleshCollection> InCollection(GetValue<DataType>(Context, &Collection).NewCopy<FFleshCollection>());
+		TArray<FIntVector4> TetElements;
+		TArray<FVector> TetVertices;
+		RadialTetMesh(InnerRadius, OuterRadius, Height, RadialSample, AngularSample, VerticalSample, TetElements, TetVertices);
+		TArray<FIntVector3> SurfaceElements = Dataflow::GetSurfaceTriangles(TetElements);;
+		TUniquePtr<FTetrahedralCollection> RadialCollection(FTetrahedralCollection::NewTetrahedralCollection(TetVertices, SurfaceElements, TetElements));
+		InCollection->AppendGeometry(*RadialCollection.Get());
+
+		SetValue<DataType>(Context, *(FManagedArrayCollection*)InCollection.Get(), &Collection);
+	}
+}
+
diff --git a/Engine/Plugins/Experimental/ChaosFlesh/Source/ChaosFleshNodes/Private/Dataflow/ChaosFleshTetrahedralNodes.cpp b/Engine/Plugins/Experimental/ChaosFlesh/Source/ChaosFleshNodes/Private/Dataflow/ChaosFleshTetrahedralNodes.cpp
--- a/Engine/Plugins/Experimental/ChaosFlesh/Source/ChaosFleshNodes/Private/Dataflow/ChaosFleshTetrahedralNodes.cpp
+++ b/Engine/Plugins/Experimental/ChaosFlesh/Source/ChaosFleshNodes/Private/Dataflow/ChaosFleshTetrahedralNodes.cpp
@@ -33,42 +33,42 @@ namespace Dataflow
 		DATAFLOW_NODE_REGISTER_CREATION_FACTORY(FGenerateTetrahedralCollectionDataflowNodes);
 		DATAFLOW_NODE_REGISTER_CREATION_FACTORY(FConstructTetGridNode);
 	}
-}
 
-// Helper to get the boundary of a tet mesh, useful for debugging / verifying output
-TArray<FIntVector3> GetSurfaceTriangles(const TArray<FIntVector4>& Tets)
-{
-	// Rotate the vector so the first element is the smallest
-	auto RotVec = [](const FIntVector3& F) -> FIntVector3
+	// Helper to get the boundary of a tet mesh, useful for debugging / verifying output
+	TArray<FIntVector3> GetSurfaceTriangles(const TArray<FIntVector4>& Tets)
 	{
-		int32 MinIdx = F.X < F.Y ? (F.X < F.Z ? 0 : 2) : (F.Y < F.Z ? 1 : 2);
-		return FIntVector3(F[MinIdx], F[(MinIdx + 1) % 3], F[(MinIdx + 2) % 3]);
-	};
-	// Reverse the winding while keeping the first element unchanged
-	auto RevVec = [](const FIntVector3& F) -> FIntVector3
-	{
-		return FIntVector3(F.X, F.Z, F.Y);
-	};
+		// Rotate the vector so the first element is the smallest
+		auto RotVec = [](const FIntVector3& F) -> FIntVector3
+		{
+			int32 MinIdx = F.X < F.Y ? (F.X < F.Z ? 0 : 2) : (F.Y < F.Z ? 1 : 2);
+			return FIntVector3(F[MinIdx], F[(MinIdx + 1) % 3], F[(MinIdx + 2) % 3]);
+		};
+		// Reverse the winding while keeping the first element unchanged
+		auto RevVec = [](const FIntVector3& F) -> FIntVector3
+		{
+			return FIntVector3(F.X, F.Z, F.Y);
+		};
 
-	TSet<FIntVector3> FacesSet;
-	for (int32 TetIdx = 0; TetIdx < Tets.Num(); ++TetIdx)
-	{
-		FIntVector3 TetF[4];
-		Chaos::Utilities::GetTetFaces(Tets[TetIdx], TetF[0], TetF[1], TetF[2], TetF[3], false);
-		for (int32 SubIdx = 0; SubIdx < 4; ++SubIdx)
+		TSet<FIntVector3> FacesSet;
+		for (int32 TetIdx = 0; TetIdx < Tets.Num(); ++TetIdx)
 		{
-			FIntVector3 Key = RotVec(TetF[SubIdx]);
-			if (FacesSet.Contains(Key))
+			FIntVector3 TetF[4];
+			Chaos::Utilities::GetTetFaces(Tets[TetIdx], TetF[0], TetF[1], TetF[2], TetF[3], false);
+			for (int32 SubIdx = 0; SubIdx < 4; ++SubIdx)
 			{
-				FacesSet.Remove(Key);
-			}
-			else
-			{
-				FacesSet.Add(RevVec(Key));
+				FIntVector3 Key = RotVec(TetF[SubIdx]);
+				if (FacesSet.Contains(Key))
+				{
+					FacesSet.Remove(Key);
+				}
+				else
+				{
+					FacesSet.Add(RevVec(Key));
+				}
 			}
 		}
+		return FacesSet.Array();
 	}
-	return FacesSet.Array();
 }
 
 //=============================================================================
@@ -93,7 +93,7 @@ void FConstructTetGridNode::Evaluate(Dataflow::FContext& Context, const FDataflo
 
 		UE_LOG(LogChaosFlesh, Display, TEXT("TetGrid generated %d points and %d tetrahedra."), X.Num(), Tets.Num());
 
-		TArray<FIntVector3> Tris = GetSurfaceTriangles(Tets);
+		TArray<FIntVector3> Tris = Dataflow::GetSurfaceTriangles(Tets);
 		TUniquePtr<FTetrahedralCollection> TetCollection(
 			FTetrahedralCollection::NewTetrahedralCollection(X, Tris, Tets));
 		InCollection->AppendGeometry(*TetCollection.Get());
@@ -204,7 +204,7 @@ void FGenerateTetrahedralCollectionDataflowNodes::EvaluateIsoStuffing(
 		{
 			TArray<FVector> Vertices; Vertices.SetNumUninitialized(IsosurfaceStuffing.Vertices.Num());
 			TArray<FIntVector4> Elements; Elements.SetNumUninitialized(IsosurfaceStuffing.Tets.Num());
-			TArray<FIntVector3> SurfaceElements = GetSurfaceTriangles(IsosurfaceStuffing.Tets);
+			TArray<FIntVector3> SurfaceElements = Dataflow::GetSurfaceTriangles(IsosurfaceStuffing.Tets);
 
 			for (int32 Tdx = 0; Tdx < IsosurfaceStuffing.Tets.Num(); ++Tdx)
 			{
@@ -270,7 +270,7 @@ void FGenerateTetrahedralCollectionDataflowNodes::EvaluateTetWild(
 		UE_LOG(LogChaosFlesh, Display,TEXT("Generating tet mesh via TetWild..."));
 		if (UE::Geometry::FTetWild::ComputeTetMesh(Params, Verts, Tris, TetVerts, Tets, &Progress))
 		{
-			TArray<FIntVector3> SurfaceElements = GetSurfaceTriangles(Tets);
+			TArray<FIntVector3> SurfaceElements = Dataflow::GetSurfaceTriangles(Tets);
 			TUniquePtr<FTetrahedralCollection> TetCollection(FTetrahedralCollection::NewTetrahedralCollection(TetVerts, SurfaceElements, Tets));
 			InCollection->AppendGeometry(*TetCollection.Get());
 
diff --git a/Engine/Plugins/Experimental/ChaosFlesh/Source/ChaosFleshNodes/Public/Dataflow/ChaosFleshRadialTetrahedronNodes.h b/Engine/Plugins/Experimental/ChaosFlesh/Source/ChaosFleshNodes/Public/Dataflow/ChaosFleshRadialTetrahedronNodes.h
new file mode 100644
--- /dev/null
+++ b/Engine/Plugins/Experimental/ChaosFlesh/Source/ChaosFleshNodes/Public/Dataflow/ChaosFleshRadialTetrahedronNodes.h
@@ -0,0 +1,65 @@
+// Copyright Epic Games, Inc. All Rights Reserved.
+
+#pragma once 
+
+#include "CoreMinimal.h"
+#include "Dataflow/DataflowEngine.h"
+#include "GeometryCollection/ManagedArrayCollection.h"
+
+#include "ChaosFleshRadialTetrahedronNodes.generated.h"
+
+class UStaticMesh;
+class FFleshCollection;
+
+USTRUCT(meta = (DataflowFlesh))
+struct FRadialTetrahedronDataflowNodes : public FDataflowNode
+{
+	GENERATED_USTRUCT_BODY()
+	DATAFLOW_NODE_DEFINE_INTERNAL(FRadialTetrahedronDataflowNodes, "RadialTetrahedron", "Flesh", "")
+	DATAFLOW_NODE_RENDER_TYPE(FGeometryCollection::StaticType(), "Collection")
+
+public:
+	typedef FManagedArrayCollection DataType;
+
+	UPROPERTY(meta = (DataflowInput, DataflowOutput, DisplayName = "Collection"))
+	FManagedArrayCollection Collection;
+
+	UPROPERTY(EditAnywhere, Category = "Dataflow", meta = (ClampMin = "0.0"))
+	double InnerRadius = double(1.0);
+
+	UPROPERTY(EditAnywhere, Category = "Dataflow", meta = (ClampMin = "0.0"))
+	double OuterRadius = double(2.0);
+
+	UPROPERTY(EditAnywhere, Category = "Dataflow", meta = (ClampMin = "0.0"))
+	double Height = double(1.0);
+
+	UPROPERTY(EditAnywhere, Category = "Dataflow", meta = (ClampMin = "0"))
+	int32 RadialSample = 2;
+
+	UPROPERTY(EditAnywhere, Category = "Dataflow", meta = (ClampMin = "0"))
+	int32 AngularSample = 4;
+
+	UPROPERTY(EditAnywhere, Category = "Dataflow", meta = (ClampMin = "0"))
+	int32 VerticalSample = 2;
+
+	FRadialTetrahedronDataflowNodes(const Dataflow::FNodeParameters& InParam, FGuid InGuid = FGuid::NewGuid())
+		: FDataflowNode(InParam, InGuid)
+	{
+		RegisterInputConnection(&Collection);
+		RegisterOutputConnection(&Collection);
+	}
+
+	virtual void Evaluate(Dataflow::FContext& Context, const FDataflowOutput* Out) const override;
+};
+
+namespace Dataflow
+{
+	void ChaosFleshRadialTetrahedronNodes();
+
+
+}
+
+
+
+
+
diff --git a/Engine/Plugins/Experimental/ChaosFlesh/Source/ChaosFleshNodes/Public/Dataflow/ChaosFleshTetrahedralNodes.h b/Engine/Plugins/Experimental/ChaosFlesh/Source/ChaosFleshNodes/Public/Dataflow/ChaosFleshTetrahedralNodes.h
--- a/Engine/Plugins/Experimental/ChaosFlesh/Source/ChaosFleshNodes/Public/Dataflow/ChaosFleshTetrahedralNodes.h
+++ b/Engine/Plugins/Experimental/ChaosFlesh/Source/ChaosFleshNodes/Public/Dataflow/ChaosFleshTetrahedralNodes.h
@@ -144,6 +144,7 @@ protected:
 
 namespace Dataflow
 {
+	TArray<FIntVector3> GetSurfaceTriangles(const TArray<FIntVector4>& Tets);
 	void ChaosFleshTetrahedralNodes();
 
 
