Advancing front routine for reparenting orphans.  Add mask to tet bindings schema.  Set mask for orphaned points.

diff --git a/Engine/Plugins/Experimental/ChaosFlesh/Source/ChaosFleshNodes/Private/Dataflow/ChaosFleshBindingsNodes.cpp b/Engine/Plugins/Experimental/ChaosFlesh/Source/ChaosFleshNodes/Private/Dataflow/ChaosFleshBindingsNodes.cpp
--- a/Engine/Plugins/Experimental/ChaosFlesh/Source/ChaosFleshNodes/Private/Dataflow/ChaosFleshBindingsNodes.cpp
+++ b/Engine/Plugins/Experimental/ChaosFlesh/Source/ChaosFleshNodes/Private/Dataflow/ChaosFleshBindingsNodes.cpp
@@ -7,6 +7,7 @@
 #include "Chaos/Tetrahedron.h"
 #include "Chaos/TriangleMesh.h"
 #include "ChaosFlesh/TetrahedralCollection.h"
+#include "Containers/Map.h"
 #include "Engine/StaticMesh.h"
 #include "Dataflow/DataflowInputOutput.h"
 #include "Engine/SkeletalMesh.h"
@@ -15,6 +16,9 @@
 #include "Rendering/SkeletalMeshRenderData.h"
 #include "UObject/PrimaryAssetId.h"
 
+DEFINE_LOG_CATEGORY(LogMeshBindings);
+
+
 namespace Dataflow
 {
 	void ChaosFleshBindingsNodes()
@@ -23,44 +27,127 @@ namespace Dataflow
 	}
 }
 
+void
+BuildVertexToVertexAdjacencyBuffer(
+	const FSkeletalMeshLODRenderData& LodRenderData,
+	TArray<TArray<uint32>>& NeighborNodes)
+{
+	const FRawStaticIndexBuffer16or32Interface* IndexBuffer = LodRenderData.MultiSizeIndexContainer.GetIndexBuffer();
+	const uint32 IndexCount = IndexBuffer->Num();
+
+	const FPositionVertexBuffer& VertexBuffer = LodRenderData.StaticVertexBuffers.PositionVertexBuffer;
+	const uint32 VertexCount = LodRenderData.StaticVertexBuffers.PositionVertexBuffer.GetNumVertices();
+
+	NeighborNodes.SetNum(0); // clear, to init clean
+	NeighborNodes.SetNum(VertexCount);
+
+	int32 BaseTriangle = 0;
+	int32 BaseVertex = 0;
+	for (int32 SectionIndex = 0; SectionIndex < LodRenderData.RenderSections.Num(); ++SectionIndex)
+	{
+		FSkelMeshRenderSection const& RenderSection = LodRenderData.RenderSections[SectionIndex];
+		int32 NumTriangles = RenderSection.NumTriangles;
+		int32 NumVertices = RenderSection.NumVertices;
+
+		TArray<uint32> RedirectionArray;
+		RedirectionArray.SetNum(VertexCount);
+		TMap<FVector, int32 /*UniqueVertexIndex*/> UniqueIndexMap;
+
+		for (int32 TriangleIt = BaseTriangle; TriangleIt < BaseTriangle + NumTriangles; ++TriangleIt)
+		{
+			const uint32 V[3] =
+			{
+				IndexBuffer->Get(TriangleIt * 3 + 0),
+				IndexBuffer->Get(TriangleIt * 3 + 1),
+				IndexBuffer->Get(TriangleIt * 3 + 2)
+			};
+
+			const FVector P[3] =
+			{
+				(FVector)VertexBuffer.VertexPosition(V[0]),
+				(FVector)VertexBuffer.VertexPosition(V[1]),
+				(FVector)VertexBuffer.VertexPosition(V[2])
+			};
+
+			for (int32 i = 0; i < 3; ++i)
+			{
+				const uint32 VertexIndex = RedirectionArray[V[i]] = UniqueIndexMap.FindOrAdd(P[i], V[i]);
+				TArray<uint32>& AdjacentVertices = NeighborNodes[VertexIndex];
+				for (int32 a = 1; a < 3; ++a)
+				{
+					const uint32 AdjacentVertexIndex = V[(i + a) % 3];
+					if (VertexIndex != AdjacentVertexIndex)
+					{
+						AdjacentVertices.AddUnique(AdjacentVertexIndex);
+					}
+				}
+			}
+		}
+
+		for (int32 VertexIt = BaseVertex + 1; VertexIt < BaseVertex + NumVertices; ++VertexIt)
+		{
+			// if this vertex has a sibling we copy the data over
+			const int32 SiblingIndex = RedirectionArray[VertexIt];
+			if (SiblingIndex != VertexIt)
+			{
+				for (int32 i = 0; i < NeighborNodes[SiblingIndex].Num(); i++)
+				{
+					const uint32 OtherNode = NeighborNodes[SiblingIndex][i];
+					if (OtherNode != VertexIt)
+					{
+						NeighborNodes[VertexIt].AddUnique(OtherNode);
+					}
+				}
+			}
+		}
+
+		BaseTriangle += NumTriangles;
+		BaseVertex += NumVertices;
+	}
+}
+
 void
 FGenerateBindings::Evaluate(Dataflow::FContext& Context, const FDataflowOutput* Out) const
 {
 	if (Out->IsA<DataType>(&Collection))
 	{
-		DataType InCollection = GetValue<DataType>(Context, &Collection); // Deep copy
+		DataType OutCollection = GetValue<DataType>(Context, &Collection); // Deep copy
 
 		TManagedArray<FIntVector4>* Tetrahedron = 
-			InCollection.FindAttribute<FIntVector4>(
+			OutCollection.FindAttribute<FIntVector4>(
 				FTetrahedralCollection::TetrahedronAttribute, FTetrahedralCollection::TetrahedralGroup);
 		TManagedArray<int32>* TetrahedronStart =
-			InCollection.FindAttribute<int32>(
+			OutCollection.FindAttribute<int32>(
 				"TetrahedronStart", FGeometryCollection::GeometryGroup);
 		TManagedArray<int32>* TetrahedronCount =
-			InCollection.FindAttribute<int32>(
+			OutCollection.FindAttribute<int32>(
 				"TetrahedronCount", FGeometryCollection::GeometryGroup);
 
 		TManagedArray<FIntVector>* Triangle =
-			InCollection.FindAttribute<FIntVector>(
+			OutCollection.FindAttribute<FIntVector>(
 				"Indices", FGeometryCollection::FacesGroup);
 		TManagedArray<int32>* FacesStart =
-			InCollection.FindAttribute<int32>(
+			OutCollection.FindAttribute<int32>(
 				"FaceStart", FGeometryCollection::GeometryGroup);
 		TManagedArray<int32>* FacesCount =
-			InCollection.FindAttribute<int32>(
+			OutCollection.FindAttribute<int32>(
 				"FaceCount", FGeometryCollection::GeometryGroup);
 
 		TManagedArray<FVector3f>* Vertex = 
-			InCollection.FindAttribute<FVector3f>(
+			OutCollection.FindAttribute<FVector3f>(
 				"Vertex", "Vertices");
 
 		TObjectPtr<const USkeletalMesh> SkeletalMesh = GetValue<TObjectPtr<const USkeletalMesh>>(Context, &SkeletalMeshIn);
 		TObjectPtr<const UStaticMesh> StaticMesh = GetValue<TObjectPtr<const UStaticMesh>>(Context, &StaticMeshIn);
-		if ((SkeletalMesh || StaticMesh) && Tetrahedron && Triangle && Vertex)
+		if ((SkeletalMesh || StaticMesh) && 
+			Tetrahedron && TetrahedronStart && TetrahedronCount &&
+			Triangle && FacesStart && FacesCount &&
+			Vertex)
 		{
 			// Extract positions to bind
 			FString MeshId;
 			TArray<TArray<FVector3f>> MeshVertices;
+			TArray<TArray<TArray<uint32>>> MeshNeighborNodes;
 			if (SkeletalMesh)
 			{
 				FPrimaryAssetId Id = SkeletalMesh->GetPrimaryAssetId();
@@ -76,6 +163,7 @@ FGenerateBindings::Evaluate(Dataflow::FContext& Context, const FDataflowOutput*
 				FSkeletalMeshRenderData* RenderData = SkeletalMesh->GetResourceForRendering();
 
 				MeshVertices.SetNum(RenderData->LODRenderData.Num());
+				MeshNeighborNodes.SetNum(RenderData->LODRenderData.Num());
 				for (int32 i = 0; i < RenderData->LODRenderData.Num(); i++)
 				{
 					FSkeletalMeshLODRenderData* LODRenderData = &RenderData->LODRenderData[i];
@@ -89,6 +177,9 @@ FGenerateBindings::Evaluate(Dataflow::FContext& Context, const FDataflowOutput*
 						const FVector3f& Pos = PositionVertexBuffer.VertexPosition(j);
 						Vertices[j] = Pos;
 					}
+
+					TArray<TArray<uint32>>& NeighborNodes = MeshNeighborNodes[i];
+					BuildVertexToVertexAdjacencyBuffer(*LODRenderData, NeighborNodes);
 				}
 			}
 			else // StaticMesh
@@ -121,23 +212,6 @@ FGenerateBindings::Evaluate(Dataflow::FContext& Context, const FDataflowOutput*
 						Vertices[j] = Pos;
 					}
 				}
-
-/*				// Geom lib does conversion like this:
-				const FMeshDescription* MeshDescription =
-					StaticMesh->IsHiResMeshDescriptionValid() ?
-					StaticMesh->GetHiResMeshDescription() :
-					StaticMesh->GetMeshDescription(0);
-				TVertexAttributesConstRef<const FVector3f> VertexPositions = MeshDescription->GetVertexPositions();
-				//TVertexAttributesConstRef<FVector3f> VertexPositions = MeshDescription->GetVertexPositions();
-
-				MeshVertices.SetNum(1);
-				MeshVertices[0].Reserve(MeshDescription->Vertices().Num());
-				for (const FVertexID VertexID : MeshDescription->Vertices().GetElementIDs())
-				{
-					const FVector3f Pos = VertexPositions.Get(VertexID);
-					MeshVertices[0].Add(Pos);
-				}
-*/
 			}
 
 			for (int32 TetMeshIdx = 0; TetMeshIdx < TetrahedronStart->Num(); TetMeshIdx++)
@@ -200,7 +274,6 @@ FGenerateBindings::Evaluate(Dataflow::FContext& Context, const FDataflowOutput*
 				TArray<Chaos::FVec3> PointNormals = TetBoundaryMesh.GetPointNormals(VertexDView, false, true);
 
 				Chaos::FTriangleMesh::TBVHType<Chaos::FRealDouble> TetBoundaryBVH;
-				//TetBoundaryMesh.BuildBVH_(VertexDView, TetBoundaryBVH);
 				TetBoundaryMesh.BuildBVH(VertexDView, TetBoundaryBVH);
 
 				//
@@ -210,12 +283,18 @@ FGenerateBindings::Evaluate(Dataflow::FContext& Context, const FDataflowOutput*
 				TArray<TArray<FIntVector4>> Parents; Parents.SetNum(MeshVertices.Num());
 				TArray<TArray<FVector4f>> Weights;	 Weights.SetNum(MeshVertices.Num());
 				TArray<TArray<FVector3f>> Offsets;	 Offsets.SetNum(MeshVertices.Num());
+				TArray<TArray<float>> Masks;		 Masks.SetNum(MeshVertices.Num());
 				TArray<int32> Orphans;
+				int32 TetHits = 0;
+				int32 TriHits = 0;
+				int32 Adoptions = 0;
+				int32 NumOrphans = 0;
 				for (int32 LOD = 0; LOD < MeshVertices.Num(); LOD++)
 				{
 					Parents[LOD].SetNumUninitialized(MeshVertices[LOD].Num());
 					Weights[LOD].SetNumUninitialized(MeshVertices[LOD].Num());
 					Offsets[LOD].SetNumUninitialized(MeshVertices[LOD].Num());
+					Masks[LOD].SetNumUninitialized(MeshVertices[LOD].Num());
 
 					TArray<int32> TetIntersections; TetIntersections.Reserve(64);
 					for (int32 i = 0; i < MeshVertices[LOD].Num(); i++)
@@ -223,19 +302,30 @@ FGenerateBindings::Evaluate(Dataflow::FContext& Context, const FDataflowOutput*
 						Parents[LOD][i] = FIntVector4(INDEX_NONE);
 						Weights[LOD][i] = FVector4f(0);
 						Offsets[LOD][i] = FVector3f(0);
+						Masks[LOD][i] = 1.0;
 
 						const FVector3f& Pos = MeshVertices[LOD][i];
-						TetIntersections = TetBVH.FindAllIntersections(Chaos::TVec3<Chaos::FReal>(Pos[0], Pos[1], Pos[2]));
+						Chaos::TVec3<Chaos::FReal> PosD(Pos[0], Pos[1], Pos[2]);
+						TetIntersections = TetBVH.FindAllIntersections(PosD);
 						int32 j = 0;
 						for (; j < TetIntersections.Num(); j++)
 						{
 							const int32 TetIdx = TetIntersections[j];
-							if (Tets[TetIdx].RobustInside(Pos, -1.0e-4))
+							if (Tets[TetIdx].RobustInside(Pos, -1.0e-4)) // includes boundary
 							{
+								TetHits++;
 								Parents[LOD][i] = (*Tetrahedron)[TetIdx];
 								Chaos::TVector<Chaos::FReal, 4> WeightsD = Tets[TetIdx].GetBarycentricCoordinates(Pos);
 								Weights[LOD][i] = FVector4f(WeightsD[0], WeightsD[1], WeightsD[2], WeightsD[3]);
 								Offsets[LOD][i] = FVector3f(0);
+
+								FVector3f EmbeddedPos =
+									(*Vertex)[Parents[LOD][i][0]] * Weights[LOD][i][0] +
+									(*Vertex)[Parents[LOD][i][1]] * Weights[LOD][i][1] +
+									(*Vertex)[Parents[LOD][i][2]] * Weights[LOD][i][2] +
+									(*Vertex)[Parents[LOD][i][3]] * Weights[LOD][i][3];
+								check((Pos - EmbeddedPos).SquaredLength() < 1.0);
+
 								break;
 							}
 						}
@@ -244,13 +334,13 @@ FGenerateBindings::Evaluate(Dataflow::FContext& Context, const FDataflowOutput*
 							// This vertex didn't land inside any tetrahedra. Project it to the tet boundary surface.
 							int32 TriIdx = INDEX_NONE;
 							Chaos::FVec3 TriWeights;
-							//if (TetBoundaryMesh.SmoothProject_(
 							if (TetBoundaryMesh.SmoothProject(
 								TetBoundaryBVH,
 								VertexDView,
 								PointNormals, Pos,
 								TriIdx, TriWeights, SurfaceProjectionIterations))
 							{
+								TriHits++;
 								const FIntVector& Tri = (*Triangle)[TriIdx];
 								Parents[LOD][i][0] = Tri[0];
 								Parents[LOD][i][1] = Tri[1];
@@ -266,7 +356,7 @@ FGenerateBindings::Evaluate(Dataflow::FContext& Context, const FDataflowOutput*
 									TriWeights[0] * Vertex->GetConstArray()[Tri[0]] +
 									TriWeights[1] * Vertex->GetConstArray()[Tri[1]] +
 									TriWeights[2] * Vertex->GetConstArray()[Tri[2]];
-								Offsets[LOD][i] = Pos - EmbeddedPos;
+								Offsets[LOD][i] = EmbeddedPos - Pos;
 							}
 							else
 							{
@@ -286,22 +376,168 @@ FGenerateBindings::Evaluate(Dataflow::FContext& Context, const FDataflowOutput*
 								Offsets[LOD][i][0] = 0.0;
 								Offsets[LOD][i][1] = 0.0;
 								Offsets[LOD][i][2] = 0.0;
+
+								Masks[LOD][i] = 0.0; // Shader does skinning for this vertex
+							} // if !SmoothProject()
+						} // if !TetIntersections
+					} // end for all vertices
+
+					// 
+					// Advancing front orphan reparenting
+					//
+
+					const TArray<TArray<uint32>>& NeighborNodes = MeshNeighborNodes[LOD];
+					TSet<int32> OrphanSet(Orphans);
+					while (Orphans.Num())
+					{
+						// Find the orphan with the fewest number of orphan neighbors, and the 
+						// most non-orphans in their 1 ring.
+						int32 Orphan = INDEX_NONE;
+						int32 NumOrphanNeighbors = TNumericLimits<int32>::Max();
+						int32 NumNonOrphanNeighbors = 0;
+						for (int32 i = 0; i < Orphans.Num(); i++)
+						{
+							int32 CurrOrphan = Orphans[i];
+							const TArray<uint32>& Neighbors = NeighborNodes[CurrOrphan];
+							int32 OrphanCount = 0;
+							int32 NonOrphanCount = 0;
+							for (int32 j = 0; j < Neighbors.Num(); j++)
+							{
+								if (OrphanSet.Contains(Neighbors[j]))
+								{
+									OrphanCount++;
+								}
+								else
+								{
+									NonOrphanCount++;
+								}
+							}
+							if (OrphanCount <= NumOrphanNeighbors && NonOrphanCount > NumNonOrphanNeighbors)
+							{
+								Orphan = CurrOrphan;
+								NumOrphanNeighbors = OrphanCount;
+								NumNonOrphanNeighbors = NonOrphanCount;
 							}
 						}
-					} // end for all vertices
+						const FVector3f& Pos = MeshVertices[LOD][Orphan];
+						Chaos::TVec3<Chaos::FReal> PosD(Pos[0], Pos[1], Pos[2]);
+
+						// Use the parent simplices of non-orphan neighbors as test candidates.
+						Chaos::FReal CurrDist = TNumericLimits<Chaos::FReal>::Max();
+						const TArray<uint32>& Neighbors = NeighborNodes[Orphan];
+						bool FoundBinding = false;
+						for (int32 i = 0; i < Neighbors.Num(); i++)
+						{
+							const uint32 Neighbor = Neighbors[i];
+							if (OrphanSet.Contains(Neighbor))
+							{
+								continue;
+							}
+
+							const FIntVector4& P = Parents[LOD][Neighbor];
+							int32 NumValid = 0;
+							for (int32 j = 0; j < 4; j++)
+							{
+								NumValid += P[j] != INDEX_NONE ? 1 : 0;
+							}
+
+							if (NumValid == 0)
+							{
+								continue;
+							}
+							else if (NumValid == 4)
+							{
+								// Parent is a tetrahedron. Reconstruct rather than go looking for it.
+								Chaos::TTetrahedron<Chaos::FReal> Tet = Chaos::TTetrahedron<Chaos::FReal>(
+									(*Vertex)[P[0]], (*Vertex)[P[1]], (*Vertex)[P[2]], (*Vertex)[P[3]]);
+
+								 Chaos::TVec4<Chaos::FReal> W;
+								 Chaos::TVec3<Chaos::FReal> EmbeddedPos = Tet.FindClosestPointAndBary(PosD, W, 1.0e-4);
+								 Chaos::TVec3<Chaos::FReal> O = EmbeddedPos - PosD;
+								 Chaos::FReal Dist = O.SquaredLength();
+								 if (Dist < CurrDist)
+								 {
+									 CurrDist = Dist;
+									 Parents[LOD][Orphan] = P;
+									 Weights[LOD][Orphan] = FVector4f(W[0], W[1], W[2], W[3]);
+									 Offsets[LOD][Orphan] = FVector3f(O[0], O[1], O[2]);
+									 Masks[LOD][i] = 1.0;
+									 FoundBinding = true;
+								 }
+							}
+							else 
+							{
+								// Find tets that share all parent indices
+								for (int32 j = 0; j < Tets.Num(); j++)
+								{
+									const FIntVector4& Tet = (*Tetrahedron)[j];
+									bool Valid = true;
+									for (int32 k = 0; k < 4 && Valid; k++)
+									{
+										Valid &= P[k] == INDEX_NONE || // parent index is unused
+											P[k] == Tet[0] || // or it matches one tet vertex
+											P[k] == Tet[1] ||
+											P[k] == Tet[2] ||
+											P[k] == Tet[3];
+									}
+									if (Valid)
+									{
+										Chaos::TVec4<Chaos::FReal> W;
+										Chaos::TVec3<Chaos::FReal> EmbeddedPos = Tets[j].FindClosestPointAndBary(PosD, W, 1.0e-4);
+										Chaos::TVec3<Chaos::FReal> O = EmbeddedPos - PosD;
+										Chaos::FReal Dist = O.SquaredLength();
+										if (Dist < CurrDist)
+										{
+											CurrDist = Dist;
+											Parents[LOD][Orphan] = Tet;
+											Weights[LOD][Orphan] = FVector4f(W[0], W[1], W[2], W[3]);
+											Offsets[LOD][Orphan] = FVector3f(O[0], O[1], O[2]);
+											Masks[LOD][i] = 1.0;
+											FoundBinding = true;
+										}
+									}
+								}
+							}
+						} // end for all neighbors
+
+						// Whether or not we successfully reparented, remove the orphan from the list.
+						OrphanSet.Remove(Orphan);
+						Orphans.Remove(Orphan);
+						if (FoundBinding)
+						{
+							Adoptions++;
+						}
+						else
+						{
+							NumOrphans++;
+						}
+					}
+
+					//ELogVerbosity::Type Verbosity = Orphans.Num() > 0 ? ELogVerbosity::Error : ELogVerbosity::Display;
+					UE_LOG(LogMeshBindings, Display,
+						TEXT("'%s' - Generated mesh bindings between tet mesh index %d and render mesh of '%s' LOD %d - stats:\n"
+							"    Render vertices num: %d\n"
+							"    Vertices in tetrahedra: %d\n"
+							"    Vertices bound to tet surface: %d\n"
+							"    Orphaned vertices reparented: %d\n"
+							"    Vertices orphaned: %d"),
+						*GetName().ToString(),
+						TetMeshIdx, *MeshId, LOD,
+						MeshVertices[LOD].Num(), TetHits, TriHits, Adoptions, NumOrphans);
+
 				} // end for all LOD
 
 				// Stash bindings in the geometry collection
-				GeometryCollection::Facades::FTetrahedralBindings TetBindings(InCollection);
+				GeometryCollection::Facades::FTetrahedralBindings TetBindings(OutCollection);
 				TetBindings.DefineSchema();
 				FName MeshName(*MeshId, MeshId.Len());
 				for (int32 LOD = 0; LOD < MeshVertices.Num(); LOD++)
 				{
 					TetBindings.AddBindingsGroup(TetMeshIdx, MeshName, LOD);
-					TetBindings.SetBindingsData(Parents[LOD], Weights[LOD], Offsets[LOD]);
+					TetBindings.SetBindingsData(Parents[LOD], Weights[LOD], Offsets[LOD], Masks[LOD]);
 				}
 			} // end for TetMeshIdx
 		}
-		SetValue<DataType>(Context, InCollection, &Collection);
+		SetValue<DataType>(Context, OutCollection, &Collection);
 	}
 }
diff --git a/Engine/Plugins/Experimental/ChaosFlesh/Source/ChaosFleshNodes/Public/Dataflow/ChaosFleshBindingsNodes.h b/Engine/Plugins/Experimental/ChaosFlesh/Source/ChaosFleshNodes/Public/Dataflow/ChaosFleshBindingsNodes.h
--- a/Engine/Plugins/Experimental/ChaosFlesh/Source/ChaosFleshNodes/Public/Dataflow/ChaosFleshBindingsNodes.h
+++ b/Engine/Plugins/Experimental/ChaosFlesh/Source/ChaosFleshNodes/Public/Dataflow/ChaosFleshBindingsNodes.h
@@ -11,6 +11,8 @@
 class UStaticMesh;
 class USkeletalMesh;
 
+DECLARE_LOG_CATEGORY_EXTERN(LogMeshBindings, Verbose, All);
+
 USTRUCT(meta = (DataflowFlesh))
 struct FGenerateBindings : public FDataflowNode
 {
