ChaosFlesh - Fix for non unity builds

diff --git a/Engine/Plugins/Experimental/ChaosFlesh/Source/ChaosFleshEngine/Public/ChaosFlesh/FleshAsset.h b/Engine/Plugins/Experimental/ChaosFlesh/Source/ChaosFleshEngine/Public/ChaosFlesh/FleshAsset.h
--- a/Engine/Plugins/Experimental/ChaosFlesh/Source/ChaosFleshEngine/Public/ChaosFlesh/FleshAsset.h
+++ b/Engine/Plugins/Experimental/ChaosFlesh/Source/ChaosFleshEngine/Public/ChaosFlesh/FleshAsset.h
@@ -5,9 +5,10 @@
 #include "CoreMinimal.h"
 
 #include "ChaosFlesh/FleshCollection.h"
-#include "UObject/ObjectMacros.h"
 #include "Dataflow/DataflowNodeParameters.h"
 #include "Dataflow/DataflowEngineTypes.h"
+#include "Engine/StaticMesh.h"
+#include "UObject/ObjectMacros.h"
 
 #include "FleshAsset.generated.h"
 
@@ -15,7 +16,6 @@ class UFleshAsset;
 class UDataflow;
 class USkeletalMesh;
 class USkeleton;
-class UStaticMesh;
 
 /**
 *	FFleshAssetEdit
