Added missing generated.h includes.

diff --git a/Engine/Plugins/Experimental/ChaosFlesh/Source/ChaosFleshEngine/Public/ChaosFlesh/ChaosDeformableTypes.h b/Engine/Plugins/Experimental/ChaosFlesh/Source/ChaosFleshEngine/Public/ChaosFlesh/ChaosDeformableTypes.h
--- a/Engine/Plugins/Experimental/ChaosFlesh/Source/ChaosFleshEngine/Public/ChaosFlesh/ChaosDeformableTypes.h
+++ b/Engine/Plugins/Experimental/ChaosFlesh/Source/ChaosFleshEngine/Public/ChaosFlesh/ChaosDeformableTypes.h
@@ -5,6 +5,8 @@
 #include "UObject/ObjectMacros.h"
 #include "Engine/EngineBaseTypes.h"
 
+#include "ChaosDeformableTypes.generated.h"
+
 UENUM(BlueprintType)
 enum class EDeformableExecutionModel : uint8
 {
