Add SimToRender transform and Amount variables to the flesh deformer's deformer graph.  Implment support for multiple flesh components, and the ability to specify which flesh asset is to be associated with which deformer.

diff --git a/Engine/Plugins/Experimental/ChaosFlesh/Source/ChaosFleshEngine/Private/DataInterfaces/DIFleshDeformer.cpp b/Engine/Plugins/Experimental/ChaosFlesh/Source/ChaosFleshEngine/Private/DataInterfaces/DIFleshDeformer.cpp
--- a/Engine/Plugins/Experimental/ChaosFlesh/Source/ChaosFleshEngine/Private/DataInterfaces/DIFleshDeformer.cpp
+++ b/Engine/Plugins/Experimental/ChaosFlesh/Source/ChaosFleshEngine/Private/DataInterfaces/DIFleshDeformer.cpp
@@ -34,7 +34,6 @@ TArray<FOptimusCDIPinDefinition>
 UDIFleshDeformer::GetPinDefinitions() const
 {
 	TArray<FOptimusCDIPinDefinition> Defs;
-	Defs.Add({ "Amount", "ReadAmount" });
 	Defs.Add({ "Mask", "ReadMask", Optimus::DomainName::Vertex, "ReadNumVertices" });
 	Defs.Add({ "Parents", "ReadParents", Optimus::DomainName::Vertex, "ReadNumVertices" });
 	Defs.Add({ "Weights", "ReadWeights", Optimus::DomainName::Vertex, "ReadNumVertices"  });
@@ -56,10 +55,6 @@ UDIFleshDeformer::GetSupportedInputs(TArray<FShaderFunctionDefinition>& OutFunct
 		.SetName(TEXT("ReadNumVertices"))
 		.AddReturnType(EShaderFundamentalType::Uint);
 
-	OutFunctions.AddDefaulted_GetRef()
-		.SetName(TEXT("ReadAmount"))
-		.AddReturnType(EShaderFundamentalType::Float);
-
 	OutFunctions.AddDefaulted_GetRef()
 		.SetName(TEXT("ReadMask"))
 		.AddReturnType(EShaderFundamentalType::Float)
@@ -96,7 +91,6 @@ BEGIN_SHADER_PARAMETER_STRUCT(FFleshDeformerDataInterfaceParameters, )
 	SHADER_PARAMETER(uint32, NumWeightsBuffer)
 	SHADER_PARAMETER(uint32, NumOffsetBuffer)
 	SHADER_PARAMETER(uint32, NumMaskBuffer)
-	SHADER_PARAMETER(float, Amount)
 	SHADER_PARAMETER_RDG_BUFFER_SRV(Buffer<float>, TetRestVertexBuffer)
 	SHADER_PARAMETER_RDG_BUFFER_SRV(Buffer<float>, TetVertexBuffer)
 	SHADER_PARAMETER_RDG_BUFFER_SRV(Buffer<int>, ParentsBuffer)
@@ -154,21 +148,60 @@ UDIFleshDeformer::CreateDataProvider(TObjectPtr<UObject> InBinding, uint64 InInp
 	{
 		AActor* Actor = Provider->SkinnedMesh->GetOwner();
 		const TSet<UActorComponent*>& Components = Actor->GetComponents();
+
+		// Count the number of flesh components we have.
+		int32 NumFleshComponents = 0;
+		for (TSet<UActorComponent*>::TConstIterator It = Components.CreateConstIterator(); It; ++It)
+		{
+			NumFleshComponents += const_cast<UFleshComponent*>(Cast<UFleshComponent>(*It)) != nullptr;
+		}
+
+		// Find the flesh component that is associated with this skeletal mesh.
 		for (TSet<UActorComponent*>::TConstIterator It = Components.CreateConstIterator(); It; ++It)
 		{
 			if (const UActorComponent* Component = *It)
 			{
-				if ((Provider->FleshMesh = const_cast<UFleshComponent*>(Cast<UFleshComponent>(Component))))
+				if (UFleshComponent* FleshComponent = const_cast<UFleshComponent*>(Cast<UFleshComponent>(Component)))
 				{
-					// TODO: Validate that this FleshMesh has bindings for this SkeletalMesh
-					break;
+					if (const UFleshAsset* FleshAsset = FleshComponent->GetRestCollection())
+					{
+						//if ((USkinnedMeshComponent*)FleshComponent->TargetSkeletalMesh == Provider->SkinnedMesh)
+						{
+							if (FleshAsset->TargetSkeletalMesh.Equals(Provider->SkinnedMesh->GetName()))
+							{
+								Provider->FleshMesh = FleshComponent;
+								break;
+							}
+							else if (FleshAsset->TargetSkeletalMesh.IsEmpty() && NumFleshComponents == 1)
+							{
+								// The typical case is probably that we only have 1 skeletal mesh and 1 flesh
+								// component.  So, in an effort to keep the easy case easy, if the TargetSkeletalMesh
+								// isn't set, just do the obvious thing and use it.
+								UE_LOG(LogFleshDeformer, Display,
+									TEXT("FleshAsset '%s' has no value for 'TargetSkeletalMesh', "
+										"using it for skeletal mesh '%s' anyway as it's the only one available."),
+									*FleshAsset->GetName(),
+									*Provider->SkinnedMesh->GetName());
+								Provider->FleshMesh = FleshComponent;
+								break;
+							}
+							else if (FleshAsset->TargetSkeletalMesh.IsEmpty())
+							{
+								UE_LOG(LogFleshDeformer, Warning,
+									TEXT("FleshAsset '%s' has no value for 'TargetSkeletalMesh', "
+										"not using it for skeletal mesh '%s' as there are other candidates."),
+									*FleshAsset->GetName(),
+									*Provider->SkinnedMesh->GetName());
+							}
+						}
+					}
 				}
-				// TODO: Someday, if/when Optimus supports deforming static meshes...
-				//if (Provider->StaticMesh = Cast<UStaticMeshComponent>(InBinding))
-				//{}
 			}
 		}
 	}
+	// TODO: Someday, if/when Optimus supports deforming static meshes...
+	//if (Provider->StaticMesh = Cast<UStaticMeshComponent>(InBinding))
+	//{}
 
 	Provider->FleshDeformerParameters = FleshDeformerParameters;
 	return Provider;
@@ -257,12 +290,15 @@ FDIFleshDeformerProviderProxy::IsValid(FValidationData const& InValidationData)
 struct FFleshDeformerDataInterfacePermutationIds
 {
 	uint32 EnableDeformerFlesh = 0;
+	uint32 EnableSurfaceOffsets = 0;
 
 	FFleshDeformerDataInterfacePermutationIds(FComputeKernelPermutationVector const& PermutationVector)
 	{
-		static FString Name(TEXT("ENABLE_DEFORMER_FLESH"));
-		static uint32 Hash = GetTypeHash(Name);
-		EnableDeformerFlesh = PermutationVector.GetPermutationBits(Name, Hash, 1);
+		{
+			static FString Name(TEXT("ENABLE_DEFORMER_FLESH"));
+			static uint32 Hash = GetTypeHash(Name);
+			EnableDeformerFlesh = PermutationVector.GetPermutationBits(Name, Hash, 1);
+		}
 	}
 };
 
@@ -495,6 +531,7 @@ FDIFleshDeformerProviderProxy::GatherDispatchData(FDispatchData const& InDispatc
 	FSkeletalMeshRenderData const& SkeletalMeshRenderData = SkeletalMeshObject->GetSkeletalMeshRenderData();
 	FSkeletalMeshLODRenderData const* LodRenderData = &SkeletalMeshRenderData.LODRenderData[LodIndex];
 	const TStridedView<FParameters> ParameterArray = MakeStridedParameterView<FParameters>(InDispatchData);
+
 	for (int32 InvocationIndex = 0; InvocationIndex < ParameterArray.Num(); ++InvocationIndex)
 	{
 		FSkelMeshRenderSection const& RenderSection = LodRenderData->RenderSections[InvocationIndex];
@@ -504,8 +541,6 @@ FDIFleshDeformerProviderProxy::GatherDispatchData(FDispatchData const& InDispatc
 		Parameters.NumVertices = RenderSection.GetNumVertices();
 		Parameters.BaseVertexIndex = RenderSection.BaseVertexIndex;
 
-		Parameters.Amount = FleshDeformerParameters.Amount;
-
 		Parameters.NumTetRestVertexBuffer = NumTetRestVertices;
 		Parameters.NumTetVertexBuffer = NumTetVertices;
 
diff --git a/Engine/Plugins/Experimental/ChaosFlesh/Source/ChaosFleshEngine/Public/ChaosFlesh/FleshAsset.h b/Engine/Plugins/Experimental/ChaosFlesh/Source/ChaosFleshEngine/Public/ChaosFlesh/FleshAsset.h
--- a/Engine/Plugins/Experimental/ChaosFlesh/Source/ChaosFleshEngine/Public/ChaosFlesh/FleshAsset.h
+++ b/Engine/Plugins/Experimental/ChaosFlesh/Source/ChaosFleshEngine/Public/ChaosFlesh/FleshAsset.h
@@ -119,7 +119,11 @@ public:
 	UPROPERTY(EditAnywhere, Category = "Geometry")
 	TObjectPtr<UStaticMesh> StaticMesh;
 
-
+	//
+	// Deformer Target Mesh
+	//
+	UPROPERTY(EditAnywhere, Category = "Deformer")
+	FString TargetSkeletalMesh;
 
 #if WITH_EDITORONLY_DATA
 	UPROPERTY(EditAnywhere, Category = "Render")
diff --git a/Engine/Plugins/Experimental/ChaosFlesh/Source/ChaosFleshEngine/Public/DataInterfaces/DIFleshDeformer.h b/Engine/Plugins/Experimental/ChaosFlesh/Source/ChaosFleshEngine/Public/DataInterfaces/DIFleshDeformer.h
--- a/Engine/Plugins/Experimental/ChaosFlesh/Source/ChaosFleshEngine/Public/DataInterfaces/DIFleshDeformer.h
+++ b/Engine/Plugins/Experimental/ChaosFlesh/Source/ChaosFleshEngine/Public/DataInterfaces/DIFleshDeformer.h
@@ -19,15 +19,9 @@ USTRUCT()
 struct FFleshDeformerParameters
 {
 	GENERATED_BODY();
-
-	/** Post-sim amount. */
-	UPROPERTY(EditAnywhere, Category = Deformer, meta = (ClampMin = "0.0", ClampMax = "1.0"))
-	float Amount = 1.0f;
-
 };
 
 
-
 /**
 * Compute Framework Data Interface for reading skeletal mesh and tetrahedral mesh bindings.
 */
diff --git a/Engine/Plugins/Experimental/ChaosFlesh/Source/ChaosFleshEngine/Shaders/Private/DIFleshDeformer.ush b/Engine/Plugins/Experimental/ChaosFlesh/Source/ChaosFleshEngine/Shaders/Private/DIFleshDeformer.ush
--- a/Engine/Plugins/Experimental/ChaosFlesh/Source/ChaosFleshEngine/Shaders/Private/DIFleshDeformer.ush
+++ b/Engine/Plugins/Experimental/ChaosFlesh/Source/ChaosFleshEngine/Shaders/Private/DIFleshDeformer.ush
@@ -11,8 +11,6 @@ uint {DataInterfaceName}_NumWeightsBuffer;
 uint {DataInterfaceName}_NumOffsetBuffer;
 uint {DataInterfaceName}_NumMaskBuffer;
 
-float {DataInterfaceName}_Amount;
-
 Buffer<float> {DataInterfaceName}_TetRestVertexBuffer;
 Buffer<float> {DataInterfaceName}_TetVertexBuffer;
 
@@ -26,11 +24,6 @@ uint ReadNumVertices_{DataInterfaceName}(uint Index)
 	return {DataInterfaceName}_NumVertices;
 }
 
-float ReadAmount_{DataInterfaceName}()
-{
-	return {DataInterfaceName}_Amount;
-}
-
 //
 // Tet mesh points, rest and dynamic 
 //
