ChaosFlesh: Bug fix for radial tetrahedral meshes.

Ryan for Yushan.

diff --git a/Engine/Plugins/Experimental/ChaosFlesh/Source/ChaosFlesh/Private/Meshing/ChaosFleshRadialMeshing.cpp b/Engine/Plugins/Experimental/ChaosFlesh/Source/ChaosFlesh/Private/Meshing/ChaosFleshRadialMeshing.cpp
--- a/Engine/Plugins/Experimental/ChaosFlesh/Source/ChaosFlesh/Private/Meshing/ChaosFleshRadialMeshing.cpp
+++ b/Engine/Plugins/Experimental/ChaosFlesh/Source/ChaosFlesh/Private/Meshing/ChaosFleshRadialMeshing.cpp
@@ -75,7 +75,7 @@ void RegularHexMesh2TetMesh(const TArray<FVector>& HexVertices, const TArray<int
 	if (ensureMsgf(HexElements.Num() % 8 == 0, TEXT("The hex mesh is not divisible by 8."))) 
 	{
 		TArray<FIntVector4> FaceToNodes = { {int32(0), int32(1), int32(2), int32(3)},{int32(0), int32(1), int32(5), int32(4)},{int32(1), int32(2), int32(6), int32(5)},{int32(3), int32(2), int32(6), int32(7)},{int32(0), int32(3), int32(7), int32(4)},{int32(4), int32(5), int32(6), int32(7)} };
-		TArray<FUintVector2> CommonFaces;
+		TArray<FIntVector2> CommonFaces;
 		int32 Offset = HexVertices.Num();
 		TetVertices = HexVertices;
 		TetElements.SetNum(0);
@@ -102,7 +102,7 @@ void RegularHexMesh2TetMesh(const TArray<FVector>& HexVertices, const TArray<int
 			f0 = int32(CommonFaces[f][0]) % 6;
 			h0 = Offset + q0;
 			// The same face is shared by two hexs
-			if (CommonFaces[f][1] > uint32(-1))
+			if (CommonFaces[f][1] > int32(-1))
 			{
 				q1 = int32(CommonFaces[f][1]) / 6;
 				h1 = Offset + q1;
@@ -143,7 +143,7 @@ void RegularHexMesh2TetMesh(const TArray<FVector>& HexVertices, const TArray<int
 	}
 }
 
-void ComputeHexMeshFaces(const TArray<int32>& HexElements, TArray<FUintVector2>& CommonFaces)
+void ComputeHexMeshFaces(const TArray<int32>& HexElements, TArray<FIntVector2>& CommonFaces)
 {
 	auto FaceGreaterThan = [](FIntVector4& i1, FIntVector4& i2)
 	{
@@ -181,7 +181,7 @@ void ComputeHexMeshFaces(const TArray<int32>& HexElements, TArray<FUintVector2>&
 }
 
 template <typename Func1, typename Func2, typename Func3>
-void ComputeMeshFaces(const TArray<int32>& mesh, Func1 GreaterThan, Func2 Equal, TArray<FUintVector2>& CommonFaces, Func3 GenerateFace, int32 PointsPerFace)
+void ComputeMeshFaces(const TArray<int32>& mesh, Func1 GreaterThan, Func2 Equal, TArray<FIntVector2>& CommonFaces, Func3 GenerateFace, int32 PointsPerFace)
 {
 	TArray<int32> Faces;
 	GenerateFace(mesh, Faces);
diff --git a/Engine/Plugins/Experimental/ChaosFlesh/Source/ChaosFlesh/Public/Meshing/ChaosFleshRadialMeshing.h b/Engine/Plugins/Experimental/ChaosFlesh/Source/ChaosFlesh/Public/Meshing/ChaosFleshRadialMeshing.h
--- a/Engine/Plugins/Experimental/ChaosFlesh/Source/ChaosFlesh/Public/Meshing/ChaosFleshRadialMeshing.h
+++ b/Engine/Plugins/Experimental/ChaosFlesh/Source/ChaosFlesh/Public/Meshing/ChaosFleshRadialMeshing.h
@@ -5,8 +5,8 @@
 void CHAOSFLESH_API RadialTetMesh(const FVector::FReal InnerRadius, const FVector::FReal OuterRadius, const FVector::FReal Height, const int32 RadialSample, const int32 AngularSample, const int32 VerticalSample, TArray<FIntVector4>& TetElements, TArray<FVector>& TetVertices);
 void CHAOSFLESH_API RadialHexMesh(const FVector::FReal InnerRadius, const FVector::FReal OuterRadius, const FVector::FReal Height, const int32 RadialSample, const int32 AngularSample, const int32 VerticalSample, TArray<int32>& HexElements, TArray<FVector>& HexVertices);
 void CHAOSFLESH_API RegularHexMesh2TetMesh(const TArray<FVector>& HexVertices, const TArray<int32>& HexElements, TArray<FVector>& TetVertices, TArray<FIntVector4>& TetElements);
-void CHAOSFLESH_API ComputeHexMeshFaces(const TArray<int32>& HexElements, TArray<FUintVector2>& CommonFaces);
+void CHAOSFLESH_API ComputeHexMeshFaces(const TArray<int32>& HexElements, TArray<FIntVector2>& CommonFaces);
 
 template <typename Func1, typename Func2, typename Func3>
-void CHAOSFLESH_API ComputeMeshFaces(const TArray<int32>& mesh, Func1 GreaterThan, Func2 Equal, TArray<FUintVector2>& CommonFaces, Func3 GenerateFace, int32 PointsPerFace);
+void CHAOSFLESH_API ComputeMeshFaces(const TArray<int32>& mesh, Func1 GreaterThan, Func2 Equal, TArray<FIntVector2>& CommonFaces, Func3 GenerateFace, int32 PointsPerFace);
 void CHAOSFLESH_API GenerateHexFaceMesh(const TArray<int32>& HexElements, TArray<int32>& Faces);
\ No newline at end of file
diff --git a/Engine/Plugins/Experimental/ChaosFlesh/Source/ChaosFleshNodes/Private/Dataflow/ChaosFleshRadialTetrahedronNodes.cpp b/Engine/Plugins/Experimental/ChaosFlesh/Source/ChaosFleshNodes/Private/Dataflow/ChaosFleshRadialTetrahedronNodes.cpp
--- a/Engine/Plugins/Experimental/ChaosFlesh/Source/ChaosFleshNodes/Private/Dataflow/ChaosFleshRadialTetrahedronNodes.cpp
+++ b/Engine/Plugins/Experimental/ChaosFlesh/Source/ChaosFleshNodes/Private/Dataflow/ChaosFleshRadialTetrahedronNodes.cpp
@@ -45,7 +45,7 @@ void FRadialTetrahedronDataflowNodes::Evaluate(Dataflow::FContext& Context, cons
 		TArray<FVector> TetVertices;
 		RadialTetMesh(InnerRadius, OuterRadius, Height, RadialSample, AngularSample, VerticalSample, TetElements, TetVertices);
 		TArray<FIntVector3> SurfaceElements = Dataflow::GetSurfaceTriangles(TetElements);;
-		TUniquePtr<FTetrahedralCollection> RadialCollection(FTetrahedralCollection::NewTetrahedralCollection(TetVertices, SurfaceElements, TetElements, false));
+		TUniquePtr<FTetrahedralCollection> RadialCollection(FTetrahedralCollection::NewTetrahedralCollection(TetVertices, SurfaceElements, TetElements));
 		InCollection->AppendGeometry(*RadialCollection.Get());
 
 		SetValue<DataType>(Context, *(FManagedArrayCollection*)InCollection.Get(), &Collection);
