Fix retention of tetrahedral incident elements.  Simplfy per-particle mass calculation, and bindings orphan reparenting with use of incident elements.  Add stats reporting to mass node.

diff --git a/Engine/Plugins/Experimental/ChaosFlesh/Source/ChaosFlesh/Private/ChaosFlesh/TetrahedralCollection.cpp b/Engine/Plugins/Experimental/ChaosFlesh/Source/ChaosFlesh/Private/ChaosFlesh/TetrahedralCollection.cpp
--- a/Engine/Plugins/Experimental/ChaosFlesh/Source/ChaosFlesh/Private/ChaosFlesh/TetrahedralCollection.cpp
+++ b/Engine/Plugins/Experimental/ChaosFlesh/Source/ChaosFlesh/Private/ChaosFlesh/TetrahedralCollection.cpp
@@ -139,6 +139,26 @@ int32 FTetrahedralCollection::AppendGeometry(
 		TetrahedronCount[NumGeometry + Idx] = Other.TetrahedronCount[Idx];
 	}
 
+	// --- VERTICES GROUP --
+	int32 NumVertices = VerticesIndex;
+	int32 OtherNumVertices = Other.IncidentElements.Num();
+	for (int32 Idx = 0; Idx < OtherNumVertices; Idx++)
+	{
+		const TArray<int32>& OtherIncidentElements = Other.IncidentElements[Idx];
+		TArray<int32>& ThisIncidentElements = IncidentElements[VerticesIndex + Idx];
+		ThisIncidentElements.SetNumUninitialized(OtherIncidentElements.Num());
+		for (int32 i = 0; i < ThisIncidentElements.Num(); i++)
+		{
+			// Offset by the number of tets we started with
+			ThisIncidentElements[i] = OtherIncidentElements[i] + NumTets;
+		}
+
+		const TArray<int32>& OtherIncidentElementsLocalIndex = Other.IncidentElementsLocalIndex[Idx];
+		TArray<int32>& ThisIncidentElementsLocalIndex = IncidentElementsLocalIndex[Idx];
+		// No offset necessary, just copy
+		ThisIncidentElementsLocalIndex = OtherIncidentElementsLocalIndex;
+	}
+
 	return ID;
 }
 
diff --git a/Engine/Plugins/Experimental/ChaosFlesh/Source/ChaosFlesh/Public/ChaosFlesh/TetrahedralCollection.h b/Engine/Plugins/Experimental/ChaosFlesh/Source/ChaosFlesh/Public/ChaosFlesh/TetrahedralCollection.h
--- a/Engine/Plugins/Experimental/ChaosFlesh/Source/ChaosFlesh/Public/ChaosFlesh/TetrahedralCollection.h
+++ b/Engine/Plugins/Experimental/ChaosFlesh/Source/ChaosFlesh/Public/ChaosFlesh/TetrahedralCollection.h
@@ -73,7 +73,7 @@ public:
 
 	/**
 	* Incident Elements Attribute
-	* For each vertex, a list of tetrahedra that includes that vertex, and the vertex's index in the tet.
+	* For each vertex, a list of tetrahedra that includes that vertex.
 	*/
 	static const FName IncidentElementsAttribute;
 	TManagedArray<TArray<int32>> IncidentElements;
diff --git a/Engine/Plugins/Experimental/ChaosFlesh/Source/ChaosFleshNodes/Private/Dataflow/ChaosFleshBindingsNodes.cpp b/Engine/Plugins/Experimental/ChaosFlesh/Source/ChaosFleshNodes/Private/Dataflow/ChaosFleshBindingsNodes.cpp
--- a/Engine/Plugins/Experimental/ChaosFlesh/Source/ChaosFleshNodes/Private/Dataflow/ChaosFleshBindingsNodes.cpp
+++ b/Engine/Plugins/Experimental/ChaosFlesh/Source/ChaosFleshNodes/Private/Dataflow/ChaosFleshBindingsNodes.cpp
@@ -118,10 +118,13 @@ FGenerateBindings::Evaluate(Dataflow::FContext& Context, const FDataflowOutput*
 				FTetrahedralCollection::TetrahedronAttribute, FTetrahedralCollection::TetrahedralGroup);
 		TManagedArray<int32>* TetrahedronStart =
 			OutCollection.FindAttribute<int32>(
-				"TetrahedronStart", FGeometryCollection::GeometryGroup);
+				FTetrahedralCollection::TetrahedronStartAttribute, FGeometryCollection::GeometryGroup);
 		TManagedArray<int32>* TetrahedronCount =
 			OutCollection.FindAttribute<int32>(
-				"TetrahedronCount", FGeometryCollection::GeometryGroup);
+				FTetrahedralCollection::TetrahedronCountAttribute, FGeometryCollection::GeometryGroup);
+		TManagedArray<TArray<int32>>* IncidentElements =
+			OutCollection.FindAttribute<TArray<int32>>(
+				FTetrahedralCollection::IncidentElementsAttribute, FGeometryCollection::VerticesGroup);
 
 		TManagedArray<FIntVector>* Triangle =
 			OutCollection.FindAttribute<FIntVector>(
@@ -220,13 +223,14 @@ FGenerateBindings::Evaluate(Dataflow::FContext& Context, const FDataflowOutput*
 				const int32 TetMeshCount = (*TetrahedronCount)[TetMeshIdx];
 
 				// Build Tetrahedra
-				TArray<Chaos::TTetrahedron<Chaos::FReal>> Tets;
+				TArray<Chaos::TTetrahedron<Chaos::FReal>> Tets;			// Index 0 == TetMeshStart
 				TArray<Chaos::TTetrahedron<Chaos::FReal>*> BVHTetPtrs;
 				Tets.SetNumUninitialized(TetMeshCount);
 				BVHTetPtrs.SetNumUninitialized(TetMeshCount);
-				for (int32 i = TetMeshStart; i < TetMeshStart + TetMeshCount; i++)
+				for (int32 i = 0; i < TetMeshCount; i++)
 				{
-					const FIntVector4& Tet = (*Tetrahedron)[i];
+					const int32 Idx = TetMeshStart + i;
+					const FIntVector4& Tet = (*Tetrahedron)[Idx];
 					Tets[i] = Chaos::TTetrahedron<Chaos::FReal>(
 						(*Vertex)[Tet[0]],
 						(*Vertex)[Tet[1]],
@@ -246,35 +250,42 @@ FGenerateBindings::Evaluate(Dataflow::FContext& Context, const FDataflowOutput*
 				// Init boundary mesh for projections.
 				//
 
+				Chaos::FTriangleMesh TetBoundaryMesh;
+				Chaos::FTriangleMesh::TBVHType<Chaos::FRealDouble> TetBoundaryBVH;
+				TArray<Chaos::FVec3> VertexD;
+				TConstArrayView<Chaos::TVec3<Chaos::FRealDouble>> VertexDView(VertexD);
+				TArray<Chaos::FVec3> PointNormals;
 				const int32 TriMeshStart = (*FacesStart)[TetMeshIdx];
 				const int32 TriMeshCount = (*FacesCount)[TetMeshIdx];
-				Chaos::FTriangleMesh TetBoundaryMesh;
-				if (TriMeshCount == Triangle->GetConstArray().Num())
-				{
-					TetBoundaryMesh.Init(
-						reinterpret_cast<const TArray<Chaos::TVec3<int32>>&>(Triangle->GetConstArray()), 0, -1, false);
-				}
-				else
+				if (bDoSurfaceProjection)
 				{
-					TArray<Chaos::TVec3<int32>> Faces; Faces.SetNumUninitialized(TriMeshCount);
-					for (int32 i = TriMeshStart; i < TriMeshStart + TriMeshCount; i++)
+					if (TriMeshCount == Triangle->GetConstArray().Num())
 					{
-						Faces[i] = reinterpret_cast<const Chaos::TVec3<int32>&>(Triangle->GetConstArray()[i]);
+						TetBoundaryMesh.Init(
+							reinterpret_cast<const TArray<Chaos::TVec3<int32>>&>(Triangle->GetConstArray()), 0, -1, false);
+					}
+					else
+					{
+						TArray<Chaos::TVec3<int32>> Faces; Faces.SetNumUninitialized(TriMeshCount);
+						for (int32 i = 0; i < TriMeshCount; i++)
+						{
+							const int32 Idx = TriMeshStart + i;
+							Faces[i] = reinterpret_cast<const Chaos::TVec3<int32>&>(Triangle->GetConstArray()[Idx]);
+						}
+						TetBoundaryMesh.Init(Faces, 0, -1, false);
 					}
-					TetBoundaryMesh.Init(Faces, 0, -1, false);
-				}
-				
-				// Promote vertices to double because that's what FTriangleMesh wants.
-				TArray<Chaos::FVec3> VertexD; VertexD.SetNumUninitialized(Vertex->Num());
-				for (int32 i = 0; i < VertexD.Num(); i++)
-				{
-					VertexD[i] = Chaos::FVec3((*Vertex)[i][0], (*Vertex)[i][1], (*Vertex)[i][2]);
-				}
-				TConstArrayView<Chaos::TVec3<Chaos::FRealDouble>> VertexDView(VertexD);
-				TArray<Chaos::FVec3> PointNormals = TetBoundaryMesh.GetPointNormals(VertexDView, false, true);
 
-				Chaos::FTriangleMesh::TBVHType<Chaos::FRealDouble> TetBoundaryBVH;
-				TetBoundaryMesh.BuildBVH(VertexDView, TetBoundaryBVH);
+					// Promote vertices to double because that's what FTriangleMesh wants.
+					VertexD.SetNumUninitialized(Vertex->Num());
+					for (int32 i = 0; i < VertexD.Num(); i++)
+					{
+						VertexD[i] = Chaos::FVec3((*Vertex)[i][0], (*Vertex)[i][1], (*Vertex)[i][2]);
+					}
+					VertexDView = TConstArrayView<Chaos::TVec3<Chaos::FRealDouble>>(VertexD);
+
+					PointNormals = TetBoundaryMesh.GetPointNormals(VertexDView, false, true);
+					TetBoundaryMesh.BuildBVH(VertexDView, TetBoundaryBVH);
+				}
 
 				//
 				// Do intersection tests against tets, then the surface.
@@ -302,7 +313,7 @@ FGenerateBindings::Evaluate(Dataflow::FContext& Context, const FDataflowOutput*
 						Parents[LOD][i] = FIntVector4(INDEX_NONE);
 						Weights[LOD][i] = FVector4f(0);
 						Offsets[LOD][i] = FVector3f(0);
-						Masks[LOD][i] = 1.0;
+						Masks[LOD][i] = 0.0; // Shader does skinning for this vertex
 
 						const FVector3f& Pos = MeshVertices[LOD][i];
 						Chaos::TVec3<Chaos::FReal> PosD(Pos[0], Pos[1], Pos[2]);
@@ -311,13 +322,14 @@ FGenerateBindings::Evaluate(Dataflow::FContext& Context, const FDataflowOutput*
 						for (; j < TetIntersections.Num(); j++)
 						{
 							const int32 TetIdx = TetIntersections[j];
-							if (Tets[TetIdx].RobustInside(Pos, -1.0e-4)) // includes boundary
+							if (!Tets[TetIdx].Outside(Pos, 1.0e-2)) // includes boundary
 							{
 								TetHits++;
-								Parents[LOD][i] = (*Tetrahedron)[TetIdx];
+								Parents[LOD][i] = (*Tetrahedron)[TetIdx + TetMeshStart];
 								Chaos::TVector<Chaos::FReal, 4> WeightsD = Tets[TetIdx].GetBarycentricCoordinates(Pos);
 								Weights[LOD][i] = FVector4f(WeightsD[0], WeightsD[1], WeightsD[2], WeightsD[3]);
 								Offsets[LOD][i] = FVector3f(0);
+								Masks[LOD][i] = 1.0; // Shader does sim for this vertex
 
 								FVector3f EmbeddedPos =
 									(*Vertex)[Parents[LOD][i][0]] * Weights[LOD][i][0] +
@@ -334,14 +346,15 @@ FGenerateBindings::Evaluate(Dataflow::FContext& Context, const FDataflowOutput*
 							// This vertex didn't land inside any tetrahedra. Project it to the tet boundary surface.
 							int32 TriIdx = INDEX_NONE;
 							Chaos::FVec3 TriWeights;
-							if (TetBoundaryMesh.SmoothProject(
-								TetBoundaryBVH,
-								VertexDView,
-								PointNormals, Pos,
-								TriIdx, TriWeights, SurfaceProjectionIterations))
+							if (bDoSurfaceProjection && 
+								TetBoundaryMesh.SmoothProject(
+									TetBoundaryBVH,
+									VertexDView,
+									PointNormals, Pos,
+									TriIdx, TriWeights, SurfaceProjectionIterations))
 							{
 								TriHits++;
-								const FIntVector& Tri = (*Triangle)[TriIdx];
+								const FIntVector& Tri = (*Triangle)[TriIdx + TriMeshStart];
 								Parents[LOD][i][0] = Tri[0];
 								Parents[LOD][i][1] = Tri[1];
 								Parents[LOD][i][2] = Tri[2];
@@ -357,6 +370,8 @@ FGenerateBindings::Evaluate(Dataflow::FContext& Context, const FDataflowOutput*
 									TriWeights[1] * Vertex->GetConstArray()[Tri[1]] +
 									TriWeights[2] * Vertex->GetConstArray()[Tri[2]];
 								Offsets[LOD][i] = EmbeddedPos - Pos;
+
+								Masks[LOD][i] = 1.0; // Shader does sim for this vertex
 							}
 							else
 							{
@@ -388,7 +403,7 @@ FGenerateBindings::Evaluate(Dataflow::FContext& Context, const FDataflowOutput*
 
 					const TArray<TArray<uint32>>& NeighborNodes = MeshNeighborNodes[LOD];
 					TSet<int32> OrphanSet(Orphans);
-					while (Orphans.Num())
+					while (bDoOrphanReparenting && Orphans.Num())
 					{
 						// Find the orphan with the fewest number of orphan neighbors, and the 
 						// most non-orphans in their 1 ring.
@@ -419,6 +434,11 @@ FGenerateBindings::Evaluate(Dataflow::FContext& Context, const FDataflowOutput*
 								NumNonOrphanNeighbors = NonOrphanCount;
 							}
 						}
+						if (Orphan == INDEX_NONE)
+						{
+							// We only have orphans with no neighbors left.
+							break;
+						}
 						const FVector3f& Pos = MeshVertices[LOD][Orphan];
 						Chaos::TVec3<Chaos::FReal> PosD(Pos[0], Pos[1], Pos[2]);
 
@@ -445,55 +465,36 @@ FGenerateBindings::Evaluate(Dataflow::FContext& Context, const FDataflowOutput*
 							{
 								continue;
 							}
-							else if (NumValid == 4)
-							{
-								// Parent is a tetrahedron. Reconstruct rather than go looking for it.
-								Chaos::TTetrahedron<Chaos::FReal> Tet = Chaos::TTetrahedron<Chaos::FReal>(
-									(*Vertex)[P[0]], (*Vertex)[P[1]], (*Vertex)[P[2]], (*Vertex)[P[3]]);
-
-								 Chaos::TVec4<Chaos::FReal> W;
-								 Chaos::TVec3<Chaos::FReal> EmbeddedPos = Tet.FindClosestPointAndBary(PosD, W, 1.0e-4);
-								 Chaos::TVec3<Chaos::FReal> O = EmbeddedPos - PosD;
-								 Chaos::FReal Dist = O.SquaredLength();
-								 if (Dist < CurrDist)
-								 {
-									 CurrDist = Dist;
-									 Parents[LOD][Orphan] = P;
-									 Weights[LOD][Orphan] = FVector4f(W[0], W[1], W[2], W[3]);
-									 Offsets[LOD][Orphan] = FVector3f(O[0], O[1], O[2]);
-									 Masks[LOD][i] = 1.0;
-									 FoundBinding = true;
-								 }
-							}
 							else 
 							{
-								// Find tets that share all parent indices
-								for (int32 j = 0; j < Tets.Num(); j++)
+								// Find tets that share parent indices
+								for(int32 j=0; j < 4; j++)
 								{
-									const FIntVector4& Tet = (*Tetrahedron)[j];
-									bool Valid = true;
-									for (int32 k = 0; k < 4 && Valid; k++)
+									const int32 ParentIdx = P[j];
+									if (IncidentElements->GetConstArray().IsValidIndex(ParentIdx))
 									{
-										Valid &= P[k] == INDEX_NONE || // parent index is unused
-											P[k] == Tet[0] || // or it matches one tet vertex
-											P[k] == Tet[1] ||
-											P[k] == Tet[2] ||
-											P[k] == Tet[3];
-									}
-									if (Valid)
-									{
-										Chaos::TVec4<Chaos::FReal> W;
-										Chaos::TVec3<Chaos::FReal> EmbeddedPos = Tets[j].FindClosestPointAndBary(PosD, W, 1.0e-4);
-										Chaos::TVec3<Chaos::FReal> O = EmbeddedPos - PosD;
-										Chaos::FReal Dist = O.SquaredLength();
-										if (Dist < CurrDist)
+										const TArray<int32>& NeighborTets = (*IncidentElements)[ParentIdx];
+										for (int32 k = 0; k < NeighborTets.Num(); k++)
 										{
-											CurrDist = Dist;
-											Parents[LOD][Orphan] = Tet;
-											Weights[LOD][Orphan] = FVector4f(W[0], W[1], W[2], W[3]);
-											Offsets[LOD][Orphan] = FVector3f(O[0], O[1], O[2]);
-											Masks[LOD][i] = 1.0;
-											FoundBinding = true;
+											const int32 TetIdx = NeighborTets[k] - TetMeshStart;
+											if (ensure(Tets.IsValidIndex(TetIdx)))
+											{
+												const Chaos::FTetrahedron& Tet = Tets[TetIdx];
+
+												Chaos::TVec4<Chaos::FReal> W;
+												Chaos::TVec3<Chaos::FReal> EmbeddedPos = Tet.FindClosestPointAndBary(PosD, W, 1.0e-4);
+												Chaos::TVec3<Chaos::FReal> O = EmbeddedPos - PosD;
+												Chaos::FReal Dist = O.SquaredLength();
+												if (Dist < CurrDist)
+												{
+													CurrDist = Dist;
+													Parents[LOD][Orphan] = (*Tetrahedron)[TetIdx + TetMeshStart];
+													Weights[LOD][Orphan] = FVector4f(W[0], W[1], W[2], W[3]);
+													Offsets[LOD][Orphan] = FVector3f(O[0], O[1], O[2]);
+													Masks[LOD][i] = 1.0; // Shader does sim for this vertex
+													FoundBinding = true;
+												}
+											}
 										}
 									}
 								}
@@ -511,7 +512,8 @@ FGenerateBindings::Evaluate(Dataflow::FContext& Context, const FDataflowOutput*
 						{
 							NumOrphans++;
 						}
-					}
+					} // end while(Orphans)
+					NumOrphans += Orphans.Num();
 
 					//ELogVerbosity::Type Verbosity = Orphans.Num() > 0 ? ELogVerbosity::Error : ELogVerbosity::Display;
 					UE_LOG(LogMeshBindings, Display,
diff --git a/Engine/Plugins/Experimental/ChaosFlesh/Source/ChaosFleshNodes/Private/Dataflow/ChaosFleshEngineAssetNodes.cpp b/Engine/Plugins/Experimental/ChaosFlesh/Source/ChaosFleshNodes/Private/Dataflow/ChaosFleshEngineAssetNodes.cpp
--- a/Engine/Plugins/Experimental/ChaosFlesh/Source/ChaosFleshNodes/Private/Dataflow/ChaosFleshEngineAssetNodes.cpp
+++ b/Engine/Plugins/Experimental/ChaosFlesh/Source/ChaosFleshNodes/Private/Dataflow/ChaosFleshEngineAssetNodes.cpp
@@ -3,6 +3,7 @@
 #include "Dataflow/ChaosFleshEngineAssetNodes.h"
 
 #include "Chaos/Math/Poisson.h"
+#include "Chaos/Utilities.h"
 #include "ChaosFlesh/ChaosFlesh.h"
 #include "ChaosFlesh/FleshAsset.h"
 #include "ChaosFlesh/FleshCollection.h"
@@ -94,62 +95,100 @@ void FSetFleshDefaultPropertiesNode::Evaluate(Dataflow::FContext& Context, const
 				TManagedArray<float>& Damping = InCollection.ModifyAttribute<float>("Damping", FGeometryCollection::VerticesGroup);
 				const TManagedArray<FIntVector4>& Tetrahedron = InCollection.GetAttribute<FIntVector4>(FTetrahedralCollection::TetrahedronAttribute, FTetrahedralCollection::TetrahedralGroup);
 				const TManagedArray<FVector3f>& Vertex = InCollection.GetAttribute<FVector3f>("Vertex", "Vertices");
-				const TManagedArray<TArray<int32>>& ElemVertexIndex = InCollection.GetAttribute<TArray<int32>>(FTetrahedralCollection::IncidentElementsAttribute, FGeometryCollection::VerticesGroup);
-				const TManagedArray<TArray<int32>>& ElemLocalIndex = InCollection.GetAttribute<TArray<int32>>(FTetrahedralCollection::IncidentElementsLocalIndexAttribute, FGeometryCollection::VerticesGroup);
-
-				double Volume = 0.0;
-				TArray<float> ElementMass, Measure;
-				ElementMass.Init(0.f, 4 * TetsNum);
-				Measure.Init(0.f, TetsNum);
+				const TManagedArray<TArray<int32>>& IncidentElements = InCollection.GetAttribute<TArray<int32>>(FTetrahedralCollection::IncidentElementsAttribute, FGeometryCollection::VerticesGroup);
+				const TManagedArray<TArray<int32>>& IncidentElementsLocalIndex = InCollection.GetAttribute<TArray<int32>>(FTetrahedralCollection::IncidentElementsLocalIndexAttribute, FGeometryCollection::VerticesGroup);
+
+				double TotalVolume = 0.0;
+				TArray<float> ElementMass;
+				TArray<float> ElementVolume;
+				ElementMass.Init(0.f, TetsNum);
+				ElementVolume.Init(0.f, TetsNum);
 				for (int e = 0; e < TetsNum; e++)
 				{
 					FVector3f X0 = Vertex[Tetrahedron[e][0]];
 					FVector3f X1 = Vertex[Tetrahedron[e][1]];
 					FVector3f X2 = Vertex[Tetrahedron[e][2]];
 					FVector3f X3 = Vertex[Tetrahedron[e][3]];
-					Measure[e] = ((X1 - X0).Dot(FVector3f::CrossProduct(X3 - X0, X2 - X0))) / 6.f;
-					if (Measure[e] < 0.f)
+					ElementVolume[e] = ((X1 - X0).Dot(FVector3f::CrossProduct(X3 - X0, X2 - X0))) / 6.f;
+					if (ElementVolume[e] < 0.f)
 					{
-						Measure[e] = -Measure[e];
+						ElementVolume[e] = -ElementVolume[e];
 					}
-					Volume += Measure[e];
+					TotalVolume += ElementVolume[e];
 				}
 				for (int32 e = 0; e < TetsNum; e++)
 				{
-					float ElementNodalMass = Density * Measure[e] / 4.f;
-					for (int i = 0; i < 4; i++)
-					{
-						ElementMass[4 * e + i] = ElementNodalMass;
-					}
+					ElementMass[e] = Density * ElementVolume[e];
 				}
 
-				bool WasMassSet = false;
-				for (int32 i = 0; i < ElemVertexIndex.Num(); i++)
+				//
+				// Set per-node mass by connected volume
+				//
+
+				TSet<int32> Visited;
+				for (int32 i = 0; i < IncidentElements.Num(); i++)
 				{
-					if (ensure(ElemVertexIndex[i].Num() == ElemLocalIndex[i].Num()))
+					const TArray<int32>& IncidentElems = IncidentElements[i];
+					for (int32 j = 0; j < IncidentElems.Num(); j++)
 					{
-						for (int32 j = 0; j < ElemVertexIndex[i].Num(); j++)
+						const int32 TetIndex = IncidentElems[j];
+						if (Tetrahedron.GetConstArray().IsValidIndex(TetIndex))
 						{
-							int32 TetIndex = ElemVertexIndex[i][j];
-							if (0 <= TetIndex && TetIndex < TetsNum)
+							for (int32 k = 0; k < 4; k++)
 							{
-								int32 MassIndex = Tetrahedron[TetIndex][ElemLocalIndex[i][j]];
-								if (0 <= MassIndex && MassIndex < VertsNum)
+								const int32 MassIndex = Tetrahedron[TetIndex][k];
+								if (Mass.GetConstArray().IsValidIndex(MassIndex))
 								{
-									WasMassSet = true;
-									Mass[MassIndex] += ElementMass[ElemVertexIndex[i][j] * 4 + ElemLocalIndex[i][j]];
+									Mass[MassIndex] += ElementMass[TetIndex] / 4;
+									Visited.Add(MassIndex);
 								}
 							}
 						}
 					}
 				}
-				if (!WasMassSet && Vertex.Num())
+				
+				int32 NumSet = 0;
+				float MinV = TNumericLimits<float>::Max();
+				float MaxV = -TNumericLimits<float>::Max();
+				double AvgV = 0.0;
+				if (Visited.Num())
 				{
-					Mass.Fill(Density * Volume / Vertex.Num());
+					NumSet = Visited.Num();
+					for (TSet<int32>::TConstIterator It = Visited.CreateConstIterator(); It; ++It)
+					{
+						const int32 MassIndex = *It;
+						AvgV += Mass[MassIndex];
+						MinV = MinV < Mass[MassIndex] ? MinV : Mass[MassIndex];
+						MaxV = MaxV > Mass[MassIndex] ? MaxV : Mass[MassIndex];
+					}
+					AvgV /= NumSet;
+				}
+				else
+				{
+					MinV = MaxV = 0.0;
+				}
+
+				//
+				// If that didn't work, set a uniform mass
+				//
+
+				if (!Visited.Num() && Vertex.Num())
+				{
+					Mass.Fill(Density * TotalVolume / Vertex.Num());
+					NumSet = Mass.Num();
+					Chaos::Utilities::GetMinAvgMax(Mass.GetConstArray(), MinV, AvgV, MaxV);
 				}
 
 				Stiffness.Fill(VertexStiffness);
 				Damping.Fill(VertexDamping);
+
+				UE_LOG(LogChaosFlesh, Display,
+					TEXT("'%s' - Set mass on %d nodes:\n"
+						"    method: %s\n"
+						"    min, avg, max: %f, %f, %f"),
+					*GetName().ToString(), NumSet, 
+					(Visited.Num() > 0 ? TEXT("connected tet volume") : TEXT("uniform")),
+					MinV, AvgV, MaxV);
 			}
 		}
 		SetValue<FManagedArrayCollection>(Context, InCollection, &Collection);
diff --git a/Engine/Plugins/Experimental/ChaosFlesh/Source/ChaosFleshNodes/Public/Dataflow/ChaosFleshBindingsNodes.h b/Engine/Plugins/Experimental/ChaosFlesh/Source/ChaosFleshNodes/Public/Dataflow/ChaosFleshBindingsNodes.h
--- a/Engine/Plugins/Experimental/ChaosFlesh/Source/ChaosFleshNodes/Public/Dataflow/ChaosFleshBindingsNodes.h
+++ b/Engine/Plugins/Experimental/ChaosFlesh/Source/ChaosFleshNodes/Public/Dataflow/ChaosFleshBindingsNodes.h
@@ -13,6 +13,7 @@ class USkeletalMesh;
 
 DECLARE_LOG_CATEGORY_EXTERN(LogMeshBindings, Verbose, All);
 
+// Generate barycentric bindings (used by the FleshDeformer deformer graph) of a render surface to a tetrahedral mesh.
 USTRUCT(meta = (DataflowFlesh))
 struct FGenerateBindings : public FDataflowNode
 {
@@ -23,18 +24,30 @@ struct FGenerateBindings : public FDataflowNode
 public:
 	typedef FManagedArrayCollection DataType;
 
+	// Passthrough geometry collection. Bindings are stored as standalone groups in the \p Collection, keyed by the name of the input render mesh and all available LOD's.
 	UPROPERTY(meta = (DataflowInput, DataflowOutput, DisplayName = "Collection"))
 	FManagedArrayCollection Collection;
 
-	UPROPERTY(EditAnywhere, Category = "Dataflow")
+	// The input mesh, whose render surface is used to generate bindings.
+	UPROPERTY(EditAnywhere, Category = "Dataflow", meta = (DataflowInput, DisplayName = "StaticMesh"))
 	TObjectPtr<const UStaticMesh> StaticMeshIn = nullptr;
 
-	UPROPERTY(EditAnywhere, Category = "Dataflow")
+	// The input mesh, whose render surface is used to generate bindings.
+	UPROPERTY(EditAnywhere, Category = "Dataflow", meta = (DataflowInput, DisplayName = "SkeletalMesh"))
 	TObjectPtr<const USkeletalMesh> SkeletalMeshIn = nullptr;
 
-	UPROPERTY(EditAnywhere, Category = "Dataflow")
+	// Enable binding to the exterior hull of the tetrahedron mesh.
+	UPROPERTY(EditAnywhere, Category = "Dataflow", meta = (DisplayName = "DoSurfaceProjection"))
+	bool bDoSurfaceProjection = true;
+
+	// The maximum number of iterations to try expanding the domain while looking for surface triangles to bind to.
+	UPROPERTY(EditAnywhere, Category = "Dataflow", meta = (DisplayName = "SurfaceProjectionIterations"))
 	uint32 SurfaceProjectionIterations = 10;
 
+	// When nodes aren't contained in tetrahedra and surface projection fails, try to find suitable bindings by looking to neighboring parents.
+	UPROPERTY(EditAnywhere, Category = "Dataflow", meta = (DisplayName = "DoOrphanReparenting"))
+	bool bDoOrphanReparenting = true;
+
 	FGenerateBindings(const Dataflow::FNodeParameters& InParam, FGuid InGuid = FGuid::NewGuid())
 		: FDataflowNode(InParam, InGuid)
 	{
@@ -50,4 +63,4 @@ public:
 namespace Dataflow
 {
 	void ChaosFleshBindingsNodes();
-}
\ No newline at end of file
+}
diff --git a/Engine/Plugins/Experimental/ChaosFlesh/Source/ChaosFleshNodes/Public/Dataflow/ChaosFleshTetrahedralNodes.h b/Engine/Plugins/Experimental/ChaosFlesh/Source/ChaosFleshNodes/Public/Dataflow/ChaosFleshTetrahedralNodes.h
--- a/Engine/Plugins/Experimental/ChaosFlesh/Source/ChaosFleshNodes/Public/Dataflow/ChaosFleshTetrahedralNodes.h
+++ b/Engine/Plugins/Experimental/ChaosFlesh/Source/ChaosFleshNodes/Public/Dataflow/ChaosFleshTetrahedralNodes.h
@@ -117,10 +117,10 @@ public:
 	// Common
 	//
 
-	UPROPERTY(EditAnywhere, Category = "Dataflow")
+	UPROPERTY(EditAnywhere, Category = "Dataflow", meta = (DataflowInput, DisplayName = "StaticMesh"))
 	TObjectPtr<const UStaticMesh> StaticMesh = nullptr;
 
-	UPROPERTY(EditAnywhere, Category = "Dataflow")
+	UPROPERTY(EditAnywhere, Category = "Dataflow", meta = (DataflowInput, DisplayName = "SkeletalMesh"))
 	TObjectPtr<const USkeletalMesh> SkeletalMesh = nullptr;
 
 	UPROPERTY(meta = (DataflowInput, DataflowOutput, DisplayName = "Collection"))
diff --git a/Engine/Source/Runtime/Experimental/Chaos/Public/Chaos/Tetrahedron.h b/Engine/Source/Runtime/Experimental/Chaos/Public/Chaos/Tetrahedron.h
--- a/Engine/Source/Runtime/Experimental/Chaos/Public/Chaos/Tetrahedron.h
+++ b/Engine/Source/Runtime/Experimental/Chaos/Public/Chaos/Tetrahedron.h
@@ -173,14 +173,32 @@ namespace Chaos {
 		static bool Inside(const TArray<TTriangle<T>>& Tris, const TVec3<T>& Location, const T HalfThickness=0)
 		{
 			checkSlow(Tris.Num() == 4);
+
 			for(int32 i=0; i < Tris.Num(); i++)
-				if (Tris[i].GetPlane().SignedDistance(Location) > HalfThickness)
+				if (Inside(Tris[i].GetPlane(), Location, -HalfThickness))
 				{
 					return false;
 				}
 			return true;
 		}
 
+		bool Outside(const TVec3<T>& Location, const T HalfThickness = 0) const
+		{
+			TArray<TTriangle<T>> Tris = GetTriangles();
+			return Outside(Tris, Location, HalfThickness);
+		}
+
+		static bool Outside(const TArray<TTriangle<T>>& Tris, const TVec3<T>& Location, const T HalfThickness = 0)
+		{
+			checkSlow(Tris.Num() == 4);
+			for (int32 i = 0; i < Tris.Num(); i++)
+				if (Outside(Tris[i].GetPlane(), Location, HalfThickness))
+				{
+					return true;
+				}
+			return false;
+		}
+
 		//! \p Tolerance should be a small negative number to include boundary.
 		bool RobustInside(const TVec3<T>& Location, const T Tolerance=0) const
 		{
@@ -334,6 +352,16 @@ namespace Chaos {
 			return Value < static_cast<T>(0) ? -1 : 1;
 		}
 
+		static bool Inside(const TPlane<T, 3>& Plane, const TVec3<T>& Location, const T HalfThickness = 0)
+		{
+			return Plane.SignedDistance(Location) <= -HalfThickness;
+		}
+
+		static bool Outside(const TPlane<T, 3>& Plane, const TVec3<T>& Location, const T HalfThickness = 0)
+		{
+			return !Inside(Plane, Location, -HalfThickness);
+		}
+
 		friend FChaosArchive& operator<<(FChaosArchive& Ar, TTetrahedron);
 
 		TVec3<T> X[4];
