ChaosFlesh: kinematic body setup initialization node

diff --git a/Engine/Plugins/Experimental/ChaosFlesh/Source/ChaosFleshNodes/Private/Dataflow/ChaosFleshKinematicInitializationNodes.cpp b/Engine/Plugins/Experimental/ChaosFlesh/Source/ChaosFleshNodes/Private/Dataflow/ChaosFleshKinematicInitializationNodes.cpp
--- a/Engine/Plugins/Experimental/ChaosFlesh/Source/ChaosFleshNodes/Private/Dataflow/ChaosFleshKinematicInitializationNodes.cpp
+++ b/Engine/Plugins/Experimental/ChaosFlesh/Source/ChaosFleshNodes/Private/Dataflow/ChaosFleshKinematicInitializationNodes.cpp
@@ -11,6 +11,7 @@
 #include "Dataflow/DataflowNodeFactory.h"
 #include "GeometryCollection/Facades/CollectionVertexBoneWeightsFacade.h"
 #include "GeometryCollection/TransformCollection.h"
+#include "PhysicsEngine/PhysicsAsset.h"
 
 #include UE_INLINE_GENERATED_CPP_BY_NAME(ChaosFleshKinematicInitializationNodes)
 
@@ -20,6 +21,7 @@ namespace Dataflow
 {
 	void RegisterChaosFleshKinematicInitializationNodes()
 	{
+		DATAFLOW_NODE_REGISTER_CREATION_FACTORY(FKinematicBodySetupInitializationDataflowNode);
 		DATAFLOW_NODE_REGISTER_CREATION_FACTORY(FKinematicInitializationDataflowNode);
 		DATAFLOW_NODE_REGISTER_CREATION_FACTORY(FKinematicTetrahedralBindingsDataflowNode);
 		DATAFLOW_NODE_REGISTER_CREATION_FACTORY(FSetVerticesKinematicDataflowNode);
@@ -284,4 +286,71 @@ void FSetVerticesKinematicDataflowNode::Evaluate(Dataflow::FContext& Context, co
 		}
 		SetValue<DataType>(Context, InCollection, &Collection);
 	}
+}
+
+void FKinematicBodySetupInitializationDataflowNode::Evaluate(Dataflow::FContext& Context, const FDataflowOutput* Out) const
+{
+	if (Out->IsA<DataType>(&Collection))
+	{
+		DataType InCollection = GetValue<DataType>(Context, &Collection);
+
+		if (TManagedArray<FVector3f>* Vertices = InCollection.FindAttribute<FVector3f>("Vertex", FGeometryCollection::VerticesGroup))
+		{	
+			if (TObjectPtr<USkeletalMesh> SkeletalMesh = GetValue<TObjectPtr<USkeletalMesh>>(Context, &SkeletalMeshIn))
+			{
+				if (UPhysicsAsset* PhysicsAsset = SkeletalMesh->GetPhysicsAsset())
+				{
+					TArray<TObjectPtr<USkeletalBodySetup>> SkeletalBodySetups = PhysicsAsset->SkeletalBodySetups;
+					TArray<FTransform> ComponentPose;
+					Dataflow::Animation::GlobalTransforms(SkeletalMesh->GetRefSkeleton(), ComponentPose);
+					for (TObjectPtr<USkeletalBodySetup> BodySetup : SkeletalBodySetups)
+					{	
+						TArray<FKSphylElem> SphylElems = BodySetup->AggGeom.SphylElems;
+						int32 BoneIndex = SkeletalMesh->GetRefSkeleton().FindBoneIndex(BodySetup->BoneName);
+						if (0 < BoneIndex && BoneIndex < ComponentPose.Num())
+						{
+							TArray<int32> BoundVerts;
+							TArray<float> BoundWeights;
+							for (FKSphylElem Capsule : SphylElems)
+							{
+								for (int32 i = 0; i < Vertices->Num(); ++i)
+								{
+									float DistanceToCapsule = Capsule.GetShortestDistanceToPoint(FVector((*Vertices)[i]), Capsule.GetTransform());
+									DistanceToCapsule = Capsule.GetShortestDistanceToPoint(FVector((*Vertices)[i]), ComponentPose[BoneIndex]);
+									if (DistanceToCapsule < UE_SMALL_NUMBER)
+									{
+										if (BoundVerts.Find(i) == INDEX_NONE)
+										{
+											BoundVerts.Add(i);
+											BoundWeights.Add(1.0);
+										}
+									}
+								}
+							}
+							//get local coords of bound verts
+							typedef GeometryCollection::Facades::FKinematicBindingFacade FKinematics;
+							FKinematics Kinematics(InCollection); Kinematics.DefineSchema();
+							if (Kinematics.IsValid())
+							{
+								FKinematics::FBindingKey Binding = Kinematics.SetBoneBindings(BoneIndex, BoundVerts, BoundWeights);
+								TManagedArray<TArray<FVector3f>>& LocalPos = InCollection.AddAttribute<TArray<FVector3f>>("LocalPosition", Binding.GroupName);
+								Kinematics.AddKinematicBinding(Binding);
+
+								auto FloatVert = [](FVector3d V) { return FVector3f(V.X, V.Y, V.Z); };
+								auto DoubleVert = [](FVector3f V) { return FVector3d(V.X, V.Y, V.Z); };
+								LocalPos[Binding.Index].SetNum(BoundVerts.Num());
+								for (int32 i = 0; i < BoundVerts.Num(); i++)
+								{
+									FVector3f Temp = (*Vertices)[BoundVerts[i]];
+									LocalPos[Binding.Index][i] = FloatVert(ComponentPose[BoneIndex].InverseTransformPosition(DoubleVert(Temp)));
+								}
+							}
+						}
+					}
+				}
+				GeometryCollection::Facades::FVertexBoneWeightsFacade(InCollection).AddBoneWeightsFromKinematicBindings();
+			}
+		}
+		SetValue<DataType>(Context, InCollection, &Collection);
+	}
 }
\ No newline at end of file
diff --git a/Engine/Plugins/Experimental/ChaosFlesh/Source/ChaosFleshNodes/Public/Dataflow/ChaosFleshKinematicInitializationNodes.h b/Engine/Plugins/Experimental/ChaosFlesh/Source/ChaosFleshNodes/Public/Dataflow/ChaosFleshKinematicInitializationNodes.h
--- a/Engine/Plugins/Experimental/ChaosFlesh/Source/ChaosFleshNodes/Public/Dataflow/ChaosFleshKinematicInitializationNodes.h
+++ b/Engine/Plugins/Experimental/ChaosFlesh/Source/ChaosFleshNodes/Public/Dataflow/ChaosFleshKinematicInitializationNodes.h
@@ -123,6 +123,36 @@ public:
 	virtual void Evaluate(Dataflow::FContext& Context, const FDataflowOutput* Out) const override;
 };
 
+USTRUCT(meta = (DataflowFlesh))
+struct FKinematicBodySetupInitializationDataflowNode : public FDataflowNode
+{
+	GENERATED_USTRUCT_BODY()
+		DATAFLOW_NODE_DEFINE_INTERNAL(FKinematicBodySetupInitializationDataflowNode, "KinematicBodySetupInitialization", "Flesh", "")
+		DATAFLOW_NODE_RENDER_TYPE(FGeometryCollection::StaticType(), "Collection")
+
+public:
+	typedef FManagedArrayCollection DataType;
+
+	UPROPERTY(EditAnywhere, Category = "Dataflow")
+		FTransform Transform;
+
+	UPROPERTY(meta = (DataflowInput, DataflowOutput, DisplayName = "Collection", DataflowPassthrough = "Collection"))
+		FManagedArrayCollection Collection;
+
+	UPROPERTY(meta = (DataflowInput, DisplayName = "SkeletalMesh"))
+		TObjectPtr<USkeletalMesh> SkeletalMeshIn;
+
+	FKinematicBodySetupInitializationDataflowNode(const Dataflow::FNodeParameters& InParam, FGuid InGuid = FGuid::NewGuid())
+		: FDataflowNode(InParam, InGuid)
+	{
+		RegisterInputConnection(&Collection);
+		RegisterOutputConnection(&Collection, &Collection);
+		RegisterInputConnection(&SkeletalMeshIn);
+	}
+
+	virtual void Evaluate(Dataflow::FContext& Context, const FDataflowOutput* Out) const override;
+};
+
 namespace Dataflow
 {
 	void RegisterChaosFleshKinematicInitializationNodes();
