ChaosFlesh : Add XPBD weak constraints and position target nodes.

diff --git a/Engine/Plugins/Experimental/ChaosFlesh/Source/ChaosFleshEngine/Private/ChaosFlesh/ChaosDeformableSolverComponent.cpp b/Engine/Plugins/Experimental/ChaosFlesh/Source/ChaosFleshEngine/Private/ChaosFlesh/ChaosDeformableSolverComponent.cpp
--- a/Engine/Plugins/Experimental/ChaosFlesh/Source/ChaosFleshEngine/Private/ChaosFlesh/ChaosDeformableSolverComponent.cpp
+++ b/Engine/Plugins/Experimental/ChaosFlesh/Source/ChaosFleshEngine/Private/ChaosFlesh/ChaosDeformableSolverComponent.cpp
@@ -203,6 +203,7 @@ void UDeformableSolverComponent::Reset()
 			, Damping
 			, bEnableGravity
 			, bEnableCorotatedConstraint
+			, bEnablePositionTargets
 		}));
 
 		for (TObjectPtr<UDeformablePhysicsComponent>& DeformableComponent : DeformableComponents)
diff --git a/Engine/Plugins/Experimental/ChaosFlesh/Source/ChaosFleshEngine/Public/ChaosFlesh/ChaosDeformableSolverComponent.h b/Engine/Plugins/Experimental/ChaosFlesh/Source/ChaosFleshEngine/Public/ChaosFlesh/ChaosDeformableSolverComponent.h
--- a/Engine/Plugins/Experimental/ChaosFlesh/Source/ChaosFleshEngine/Public/ChaosFlesh/ChaosDeformableSolverComponent.h
+++ b/Engine/Plugins/Experimental/ChaosFlesh/Source/ChaosFleshEngine/Public/ChaosFlesh/ChaosDeformableSolverComponent.h
@@ -112,6 +112,9 @@ public:
 
 	UPROPERTY(EditAnywhere, Category = "Chaos Deformable")
 		bool bEnableCorotatedConstraint = true;
+
+	UPROPERTY(EditAnywhere, Category = "Chaos Deformable")
+		bool bEnablePositionTargets = true;
 	//UPROPERTY(EditAnywhere, Category = Chaos)
 	//	bool bWaitForParallelDeformableTask = true;
 
diff --git a/Engine/Plugins/Experimental/ChaosFlesh/Source/ChaosFleshNodes/Private/Dataflow/ChaosFleshNodesPlugin.cpp b/Engine/Plugins/Experimental/ChaosFlesh/Source/ChaosFleshNodes/Private/Dataflow/ChaosFleshNodesPlugin.cpp
--- a/Engine/Plugins/Experimental/ChaosFlesh/Source/ChaosFleshNodes/Private/Dataflow/ChaosFleshNodesPlugin.cpp
+++ b/Engine/Plugins/Experimental/ChaosFlesh/Source/ChaosFleshNodes/Private/Dataflow/ChaosFleshNodesPlugin.cpp
@@ -10,6 +10,7 @@
 #include "Dataflow/ChaosFleshKinematicInitializationNodes.h"
 #include "Dataflow/ChaosFleshRadialTetrahedronNodes.h"
 #include "Dataflow/ChaosFleshRenderInitializationNodes.h"
+#include "Dataflow/ChaosFleshPositionTargetInitializationNodes.h"
 #include "Dataflow/ChaosFleshTetrahedralNodes.h"
 #include "Modules/ModuleManager.h"
 
@@ -23,6 +24,7 @@ void IChaosFleshNodesPlugin::StartupModule()
 	Dataflow::ChaosFleshFiberDirectionInitializationNodes();
 	Dataflow::ChaosFleshRenderInitializationNodes();
 	Dataflow::RegisterChaosFleshKinematicInitializationNodes();
+	Dataflow::RegisterChaosFleshPositionTargetInitializationNodes();
 	Dataflow::ChaosFleshTetrahedralNodes();
 	Dataflow::ChaosFleshBindingsNodes();
 	Dataflow::ChaosFleshRadialTetrahedronNodes();
diff --git a/Engine/Plugins/Experimental/ChaosFlesh/Source/ChaosFleshNodes/Private/Dataflow/ChaosFleshPositionTargetInitializationNodes.cpp b/Engine/Plugins/Experimental/ChaosFlesh/Source/ChaosFleshNodes/Private/Dataflow/ChaosFleshPositionTargetInitializationNodes.cpp
new file mode 100644
--- /dev/null
+++ b/Engine/Plugins/Experimental/ChaosFlesh/Source/ChaosFleshNodes/Private/Dataflow/ChaosFleshPositionTargetInitializationNodes.cpp
@@ -0,0 +1,367 @@
+// Copyright Epic Games, Inc. All Rights Reserved.
+
+#include "Dataflow/ChaosFleshPositionTargetInitializationNodes.h"
+
+#include "ChaosFlesh/FleshCollection.h"
+#include "ChaosFlesh/FleshCollectionUtility.h"
+#include "Dataflow/DataflowEngineUtil.h"
+#include "Engine/SkeletalMesh.h"
+#include "Dataflow/DataflowInputOutput.h"
+#include "Chaos/BoundingVolumeHierarchy.h"
+#include "GeometryCollection/Facades/CollectionKinematicBindingFacade.h"
+#include "GeometryCollection/Facades/CollectionPositionTargetFacade.h"
+#include "Dataflow/DataflowNodeFactory.h"
+#include "GeometryCollection/Facades/CollectionVertexBoneWeightsFacade.h"
+#include "GeometryCollection/TransformCollection.h"
+#include "Chaos/Tetrahedron.h"
+
+#include UE_INLINE_GENERATED_CPP_BY_NAME(ChaosFleshPositionTargetInitializationNodes)
+
+//DEFINE_LOG_CATEGORY_STATIC(FKinematicInitializationNodesLog, Log, All);
+
+namespace Dataflow
+{
+	void RegisterChaosFleshPositionTargetInitializationNodes()
+	{
+		DATAFLOW_NODE_REGISTER_CREATION_FACTORY(FAddKinematicParticlesDataflowNode);
+		DATAFLOW_NODE_REGISTER_CREATION_FACTORY(FSetVertexVertexPositionTargetBindingDataflowNode);
+		DATAFLOW_NODE_REGISTER_CREATION_FACTORY(FSetVertexTetrahedraPositionTargetBindingDataflowNode);
+	}
+}
+
+
+void FAddKinematicParticlesDataflowNode::Evaluate(Dataflow::FContext& Context, const FDataflowOutput* Out) const
+{
+	if (Out->IsA<DataType>(&Collection))
+	{
+		DataType InCollection = GetValue<DataType>(Context, &Collection);
+
+		if (TManagedArray<FVector3f>* Vertices = InCollection.FindAttribute<FVector3f>("Vertex", FGeometryCollection::VerticesGroup))
+		{
+			TArray<int32> TargetIndices;
+			TargetIndices.SetNum(0);
+			if (FindInput(&VertexIndicesIn) && FindInput(&VertexIndicesIn)->GetConnection())
+			{
+				TArray<int32> BoundVerts;
+				TArray<float> BoundWeights;
+
+				for (int32 SelectionIndex : GetValue<TArray<int32>>(Context, &VertexIndicesIn))
+				{
+					if (0 <= SelectionIndex && SelectionIndex < Vertices->Num())
+					{
+						BoundVerts.Add(SelectionIndex);
+					}
+				}
+				if (BoundVerts.Num())
+				{
+					BoundWeights.Init(1.0, BoundVerts.Num());
+					GeometryCollection::Facades::FKinematicBindingFacade Kinematics(InCollection);
+					Kinematics.AddKinematicBinding(Kinematics.SetBoneBindings(INDEX_NONE, BoundVerts, BoundWeights));
+				}
+			}
+			else if (TObjectPtr<USkeletalMesh> SkeletalMesh = GetValue<TObjectPtr<USkeletalMesh>>(Context, &SkeletalMeshIn))
+			{
+				int32 IndexValue = GetValue<int32>(Context, &BoneIndexIn);
+				if (IndexValue != INDEX_NONE)
+				{
+					TArray<FTransform> ComponentPose;
+					Dataflow::Animation::GlobalTransforms(SkeletalMesh->GetRefSkeleton(), ComponentPose);
+
+					TArray<int32> BranchIndices;
+					if (SkeletalSelectionMode == ESkeletalSeletionMode::Dataflow_SkeletalSelection_Branch)
+					{
+						TArray<int32> ToProcess;
+						int32 CurrentIndex = IndexValue;
+						while (SkeletalMesh->GetRefSkeleton().IsValidIndex(CurrentIndex))
+						{
+							TArray<int32> Buffer;
+							SkeletalMesh->GetRefSkeleton().GetDirectChildBones(CurrentIndex, Buffer);
+
+							if (Buffer.Num())
+							{
+								ToProcess.Append(Buffer);
+							}
+
+							BranchIndices.Add(CurrentIndex);
+
+							CurrentIndex = INDEX_NONE;
+							if (ToProcess.Num())
+							{
+								CurrentIndex = ToProcess.Pop();
+							}
+						}
+					}
+					else // ESkeletalSeletionMode::Dataflow_SkeletalSelection_Single
+					{
+						BranchIndices.Add(IndexValue);
+					}
+
+					int32 ParticleIndex = InCollection.AddElements(BranchIndices.Num(), FGeometryCollection::VerticesGroup);
+					TManagedArray<FVector3f>& CurrentVertices = InCollection.ModifyAttribute<FVector3f>("Vertex", FGeometryCollection::VerticesGroup);
+
+					for (int32 Index = 0; Index < BranchIndices.Num(); Index++)
+					{
+						FVector3f BonePosition(ComponentPose[BranchIndices[Index]].GetTranslation());
+
+						CurrentVertices[ParticleIndex+Index] = BonePosition;
+
+						TArray<int32> BoundVerts;
+						TArray<float> BoundWeights;
+
+						BoundVerts.Add(ParticleIndex + Index);
+						BoundWeights.Add(1.0);
+						TargetIndices.Emplace(ParticleIndex + Index);
+
+
+						if (BoundVerts.Num())
+						{
+							GeometryCollection::Facades::FKinematicBindingFacade Kinematics(InCollection);
+							Kinematics.AddKinematicBinding(Kinematics.SetBoneBindings(Index, BoundVerts, BoundWeights));
+						}
+					}
+
+					//debugging code for tet binding:
+					//int32 ParticleIndex1 = InCollection.AddElements(1, FGeometryCollection::VerticesGroup);
+					//TManagedArray<FVector3f>& CurrentVertices1 = InCollection.ModifyAttribute<FVector3f>("Vertex", FGeometryCollection::VerticesGroup);
+					//CurrentVertices1[ParticleIndex1][0] = -40.f;
+					//CurrentVertices1[ParticleIndex1][1] = -50.f;
+					//CurrentVertices1[ParticleIndex1][2] = -10.f;
+					//GeometryCollection::Facades::FKinematicBindingFacade Kinematics(InCollection);
+					//TArray<int32> BoundVerts;
+					//TArray<float> BoundWeights;
+					//BoundVerts.Add(ParticleIndex1);
+					//BoundWeights.Add(1.0);
+					//Kinematics.AddKinematicBinding(Kinematics.SetBoneBindings(BranchIndices[0], BoundVerts, BoundWeights));
+					//TargetIndices.Emplace(ParticleIndex1);
+
+					GeometryCollection::Facades::FVertexBoneWeightsFacade(InCollection).AddBoneWeightsFromKinematicBindings();
+				}
+			}
+
+			SetValue<TArray<int32>>(Context, TargetIndices, &TargetIndicesOut);
+		}
+		SetValue<DataType>(Context, InCollection, &Collection);
+	}
+}
+
+
+void FSetVertexVertexPositionTargetBindingDataflowNode::Evaluate(Dataflow::FContext& Context, const FDataflowOutput* Out) const
+{
+	if (Out->IsA<DataType>(&Collection))
+	{
+		DataType InCollection = GetValue<DataType>(Context, &Collection);
+
+		if (TManagedArray<FVector3f>* Vertices = InCollection.FindAttribute<FVector3f>("Vertex", FGeometryCollection::VerticesGroup))
+		{
+			if (FindInput(&TargetIndicesIn) && FindInput(&TargetIndicesIn)->GetConnection())
+			{
+				Chaos::FReal SphereRadius = (Chaos::FReal)0.;
+
+				Chaos::TVec3<float> CoordMaxs(-FLT_MAX);
+				Chaos::TVec3<float> CoordMins(FLT_MAX);
+				for (int32 i = 0; i < Vertices->Num(); i++)
+				{
+					for (int32 j = 0; j < 3; j++) 
+					{
+						if ((*Vertices)[i][j] > CoordMaxs[j]) 
+						{
+							CoordMaxs[j] = (*Vertices)[i][j];
+						}
+						if ((*Vertices)[i][j] < CoordMins[j])
+						{
+							CoordMins[j] = (*Vertices)[i][j];
+						}
+					}
+				}
+
+				Chaos::TVec3<float> CoordDiff = (CoordMaxs - CoordMins) * RadiusRatio;
+
+				SphereRadius = Chaos::FReal(FGenericPlatformMath::Min(CoordDiff[0], FGenericPlatformMath::Min(CoordDiff[1], CoordDiff[2])));
+
+				TArray<Chaos::TSphere<Chaos::FReal, 3>*> VertexSpherePtrs;
+				TArray<Chaos::TSphere<Chaos::FReal, 3>> VertexSpheres;
+
+				VertexSpheres.Init(Chaos::TSphere<Chaos::FReal, 3>(Chaos::TVec3<Chaos::FReal>(0), SphereRadius), Vertices->Num());
+				VertexSpherePtrs.SetNum(Vertices->Num());
+
+				for (int32 i = 0; i < Vertices->Num(); i++)
+				{
+					Chaos::TVec3<Chaos::FReal> SphereCenter((*Vertices)[i]);
+					Chaos::TSphere<Chaos::FReal, 3> VertexSphere(SphereCenter, SphereRadius);
+					VertexSpheres[i] = Chaos::TSphere<Chaos::FReal, 3>(SphereCenter, SphereRadius);
+					VertexSpherePtrs[i] = &VertexSpheres[i];
+				}
+				Chaos::TBoundingVolumeHierarchy<
+					TArray<Chaos::TSphere<Chaos::FReal, 3>*>,
+					TArray<int32>,
+					Chaos::FReal,
+					3> VertexBVH(VertexSpherePtrs);
+
+			
+				TArray<int32> TargetIndicesLocal = GetValue<TArray<int32>>(Context, &TargetIndicesIn);
+				TSet<int32> TargetIndicesSet = TSet<int32>(TargetIndicesLocal);
+				TArray<int32> SourceIndices;
+				SourceIndices.Init(-1, TargetIndicesLocal.Num());
+				
+				for (int32 i = 0; i < TargetIndicesLocal.Num(); i++)
+				{
+					if (TargetIndicesLocal[i] > -1 && TargetIndicesLocal[i] < Vertices->Num()) 
+					{
+						FVector3f ParticlePos = (*Vertices)[TargetIndicesLocal[i]];
+						TArray<int32> VertexIntersections = VertexBVH.FindAllIntersections(ParticlePos);
+						float MinDistance = 10.f * (float)SphereRadius;
+						int32 MinIndex = -1;
+						for (int32 k = 0; k < VertexIntersections.Num(); k++)
+						{
+							if ((ParticlePos - (*Vertices)[VertexIntersections[k]]).Size() < MinDistance
+								&& VertexIntersections[k] != TargetIndicesLocal[i] && !TargetIndicesSet.Contains(VertexIntersections[k]))
+							{
+								MinIndex = VertexIntersections[k];
+								MinDistance = (ParticlePos - (*Vertices)[VertexIntersections[k]]).Size();
+							}
+						}
+						ensure(MinIndex != -1);
+						SourceIndices[i] = MinIndex;
+					}
+				}
+				GeometryCollection::Facades::FPositionTargetFacade PositionTargets(InCollection);
+				PositionTargets.DefineSchema();
+				for (int32 i = 0; i < TargetIndicesLocal.Num(); i++)
+				{
+					if (SourceIndices[i] != -1)
+					{
+						GeometryCollection::Facades::FPositionTargetsData DataPackage;
+						DataPackage.TargetIndex.Init(TargetIndicesLocal[i], 1);
+						DataPackage.TargetWeights.Init(1.f, 1);
+						DataPackage.SourceWeights.Init(1.f, 1);
+						DataPackage.SourceIndex.Init(SourceIndices[i], 1);
+						if (TManagedArray<float>* Mass = InCollection.FindAttribute<float>("Mass", FGeometryCollection::VerticesGroup))
+						{
+							if ((*Mass)[SourceIndices[i]] > 0.f)
+							{
+								DataPackage.Stiffness = PositionTargetStiffness * (*Mass)[SourceIndices[i]];
+							}
+							else
+							{
+								DataPackage.Stiffness = PositionTargetStiffness;
+							}
+						}
+						else
+						{
+							DataPackage.Stiffness = PositionTargetStiffness;
+						}
+
+						PositionTargets.AddPositionTarget(DataPackage);
+					}
+				}
+			}
+			
+		}
+		SetValue<DataType>(Context, InCollection, &Collection);
+	}
+}
+
+
+void FSetVertexTetrahedraPositionTargetBindingDataflowNode::Evaluate(Dataflow::FContext& Context, const FDataflowOutput* Out) const
+{
+	if (Out->IsA<DataType>(&Collection))
+	{
+		DataType InCollection = GetValue<DataType>(Context, &Collection);
+		if (FindInput(&TargetIndicesIn) && FindInput(&TargetIndicesIn)->GetConnection()) 
+		{
+			if (TManagedArray<FVector3f>* Vertices = InCollection.FindAttribute<FVector3f>("Vertex", FGeometryCollection::VerticesGroup))
+			{
+				if (TManagedArray<FIntVector4>* Tetrahedron = InCollection.FindAttribute<FIntVector4>(FTetrahedralCollection::TetrahedronAttribute, FTetrahedralCollection::TetrahedralGroup))
+				{
+					TManagedArray<int32>* TetrahedronStart = InCollection.FindAttribute<int32>(FTetrahedralCollection::TetrahedronStartAttribute, FGeometryCollection::GeometryGroup);
+					TManagedArray<int32>* TetrahedronCount = InCollection.FindAttribute<int32>(FTetrahedralCollection::TetrahedronCountAttribute, FGeometryCollection::GeometryGroup);
+					if (TetrahedronStart && TetrahedronCount)
+					{
+						for (int32 TetMeshIdx = 0; TetMeshIdx < TetrahedronStart->Num(); TetMeshIdx++)
+						{
+							const int32 TetMeshStart = (*TetrahedronStart)[TetMeshIdx];
+							const int32 TetMeshCount = (*TetrahedronCount)[TetMeshIdx];
+
+							// Build Tetrahedra
+							TArray<Chaos::TTetrahedron<Chaos::FReal>> Tets;			// Index 0 == TetMeshStart
+							TArray<Chaos::TTetrahedron<Chaos::FReal>*> BVHTetPtrs;
+							Tets.SetNumUninitialized(TetMeshCount);
+							BVHTetPtrs.SetNumUninitialized(TetMeshCount);
+							for (int32 i = 0; i < TetMeshCount; i++)
+							{
+								const int32 Idx = TetMeshStart + i;
+								const FIntVector4& Tet = (*Tetrahedron)[Idx];
+								Tets[i] = Chaos::TTetrahedron<Chaos::FReal>(
+									(*Vertices)[Tet[0]],
+									(*Vertices)[Tet[1]],
+									(*Vertices)[Tet[2]],
+									(*Vertices)[Tet[3]]);
+								BVHTetPtrs[i] = &Tets[i];
+							}
+
+							Chaos::TBoundingVolumeHierarchy<
+								TArray<Chaos::TTetrahedron<Chaos::FReal>*>,
+								TArray<int32>,
+								Chaos::FReal,
+								3> TetBVH(BVHTetPtrs);
+
+							TArray<int32> TargetIndicesLocal = GetValue<TArray<int32>>(Context, &TargetIndicesIn);
+							TArray<int32> SourceIndices;
+							SourceIndices.Init(-1, TargetIndicesLocal.Num());
+
+							GeometryCollection::Facades::FPositionTargetFacade PositionTargets(InCollection);
+							PositionTargets.DefineSchema();
+
+							for (int32 i = 0; i < TargetIndicesLocal.Num(); i++)
+							{
+								if (TargetIndicesLocal[i] > -1 && TargetIndicesLocal[i] < Vertices->Num())
+								{
+									FVector3f ParticlePos = (*Vertices)[TargetIndicesLocal[i]];
+									TArray<int32> TetIntersections = TetBVH.FindAllIntersections(ParticlePos);
+									for (int32 j = 0; j < TetIntersections.Num(); j++)
+									{
+										const int32 TetIdx = TetIntersections[j];
+										if (!Tets[TetIdx].Outside(ParticlePos, 1.0e-2))
+										{
+											Chaos::TVector<Chaos::FReal, 4> WeightsD = Tets[TetIdx].GetBarycentricCoordinates(ParticlePos);
+											GeometryCollection::Facades::FPositionTargetsData DataPackage;
+											DataPackage.TargetIndex.Init(TargetIndicesLocal[i], 1);
+											DataPackage.TargetWeights.Init(1.f, 1);
+											DataPackage.SourceWeights.Init(1.f, 4);
+											DataPackage.SourceIndex.Init(-1, 4);
+											DataPackage.SourceIndex[0] = (*Tetrahedron)[TetIdx][0];
+											DataPackage.SourceIndex[1] = (*Tetrahedron)[TetIdx][1];
+											DataPackage.SourceIndex[2] = (*Tetrahedron)[TetIdx][2];
+											DataPackage.SourceIndex[3] = (*Tetrahedron)[TetIdx][3];
+											DataPackage.SourceWeights[0] = WeightsD[0];
+											DataPackage.SourceWeights[1] = WeightsD[1];
+											DataPackage.SourceWeights[2] = WeightsD[2];
+											DataPackage.SourceWeights[3] = WeightsD[3];
+											if (TManagedArray<float>* Mass = InCollection.FindAttribute<float>("Mass", FGeometryCollection::VerticesGroup))
+											{
+												DataPackage.Stiffness = 0.f;
+												for (int32 k = 0; k < 4; k++)
+												{
+													DataPackage.Stiffness += DataPackage.SourceWeights[k] * PositionTargetStiffness * (*Mass)[DataPackage.SourceIndex[0]];
+												}
+											}
+											else
+											{
+												DataPackage.Stiffness = PositionTargetStiffness;
+											}
+											PositionTargets.AddPositionTarget(DataPackage);
+											break;
+										}
+									}
+								}
+							}
+						}
+					}
+				}
+			}
+		}
+		SetValue<DataType>(Context, InCollection, &Collection);
+	}
+
+
+}
\ No newline at end of file
diff --git a/Engine/Plugins/Experimental/ChaosFlesh/Source/ChaosFleshNodes/Public/Dataflow/ChaosFleshPositionTargetInitializationNodes.h b/Engine/Plugins/Experimental/ChaosFlesh/Source/ChaosFleshNodes/Public/Dataflow/ChaosFleshPositionTargetInitializationNodes.h
new file mode 100644
--- /dev/null
+++ b/Engine/Plugins/Experimental/ChaosFlesh/Source/ChaosFleshNodes/Public/Dataflow/ChaosFleshPositionTargetInitializationNodes.h
@@ -0,0 +1,176 @@
+// Copyright Epic Games, Inc. All Rights Reserved.
+
+#pragma once 
+
+#include "CoreMinimal.h"
+#include "Dataflow/DataflowCore.h"
+#include "Dataflow/DataflowEngine.h"
+#include "Dataflow/ChaosFleshKinematicInitializationNodes.h"
+#include "ChaosFleshPositionTargetInitializationNodes.generated.h"
+
+class USkeletalMesh;
+
+
+USTRUCT(meta = (DataflowFlesh))
+struct FAddKinematicParticlesDataflowNode : public FDataflowNode
+{
+	GENERATED_USTRUCT_BODY()
+		DATAFLOW_NODE_DEFINE_INTERNAL(FAddKinematicParticlesDataflowNode, "AddKinematicParticles", "Flesh", "")
+		DATAFLOW_NODE_RENDER_TYPE(FGeometryCollection::StaticType(), "Collection")
+
+public:
+	typedef FManagedArrayCollection DataType;
+
+	UPROPERTY(EditAnywhere, Category = "Dataflow")
+		float Radius = 40.f;
+
+	UPROPERTY(EditAnywhere, Category = "Dataflow")
+		FTransform Transform;
+
+	UPROPERTY(EditAnywhere, Category = "Dataflow")
+		ESkeletalSeletionMode SkeletalSelectionMode = ESkeletalSeletionMode::Dataflow_SkeletalSelection_Single;
+
+	UPROPERTY(meta = (DataflowInput, DataflowOutput, DisplayName = "Collection", DataflowPassthrough = "Collection"))
+		FManagedArrayCollection Collection;
+
+	UPROPERTY(meta = (DataflowInput, DisplayName = "SkeletalMesh"))
+		TObjectPtr<USkeletalMesh> SkeletalMeshIn;
+
+	UPROPERTY(meta = (DataflowInput, DisplayName = "SelectionSet"))
+		TArray<int32> VertexIndicesIn;
+
+	UPROPERTY(meta = (DataflowInput, DisplayName = "BoneIndex"))
+		int32 BoneIndexIn = INDEX_NONE;
+
+	UPROPERTY(meta = (DataflowOutput, DisplayName = "TargetIndices"))
+		TArray<int32> TargetIndicesOut;
+
+
+
+	FAddKinematicParticlesDataflowNode(const Dataflow::FNodeParameters& InParam, FGuid InGuid = FGuid::NewGuid())
+		: FDataflowNode(InParam, InGuid)
+	{
+		RegisterInputConnection(&Collection);
+		RegisterOutputConnection(&Collection, &Collection);
+		RegisterInputConnection(&SkeletalMeshIn);
+		RegisterInputConnection(&VertexIndicesIn);
+		RegisterInputConnection(&BoneIndexIn);
+		RegisterOutputConnection(&TargetIndicesOut);
+	}
+
+	virtual void Evaluate(Dataflow::FContext& Context, const FDataflowOutput* Out) const override;
+};
+
+
+USTRUCT(meta = (DataflowFlesh))
+struct FSetVertexVertexPositionTargetBindingDataflowNode : public FDataflowNode
+{
+	GENERATED_USTRUCT_BODY()
+		DATAFLOW_NODE_DEFINE_INTERNAL(FSetVertexVertexPositionTargetBindingDataflowNode, "SetVertexVertexPositionTargetBinding", "Flesh", "")
+		DATAFLOW_NODE_RENDER_TYPE(FGeometryCollection::StaticType(), "Collection")
+
+public:
+	typedef FManagedArrayCollection DataType;
+
+	UPROPERTY(EditAnywhere, Category = "Dataflow", meta = (ClampMin = "0.0", ClampMax = "1.0", UIMin = "0.0", UIMax = "1.0"))
+		float RadiusRatio = .1f;
+
+	UPROPERTY(meta = (DataflowInput, DataflowOutput, DisplayName = "Collection", DataflowPassthrough = "Collection"))
+		FManagedArrayCollection Collection;
+
+	UPROPERTY(meta = (DataflowInput, DisplayName = "TargetIndicesIn"))
+		TArray<int32> TargetIndicesIn;
+
+	UPROPERTY(EditAnywhere, Category = "Dataflow")
+		float PositionTargetStiffness = 10000.f;
+
+
+
+	FSetVertexVertexPositionTargetBindingDataflowNode(const Dataflow::FNodeParameters& InParam, FGuid InGuid = FGuid::NewGuid())
+		: FDataflowNode(InParam, InGuid)
+	{
+		RegisterInputConnection(&Collection);
+		RegisterOutputConnection(&Collection, &Collection);
+		RegisterInputConnection(&TargetIndicesIn);
+	}
+
+	virtual void Evaluate(Dataflow::FContext& Context, const FDataflowOutput* Out) const override;
+};
+
+//USTRUCT(meta = (DataflowFlesh))
+//struct FComputeVertexSphereBVHDataflowNode : public FDataflowNode
+//{
+//	GENERATED_USTRUCT_BODY()
+//		DATAFLOW_NODE_DEFINE_INTERNAL(FComputeVertexSphereBVHDataflowNode, "ComputeVertexSphereBVH", "Flesh", "")
+//		DATAFLOW_NODE_RENDER_TYPE(FGeometryCollection::StaticType(), "Collection")
+//
+//public:
+//	typedef FManagedArrayCollection DataType;
+//
+//	UPROPERTY(EditAnywhere, Category = "Dataflow")
+//		float Radius = 40.f;
+//
+//	UPROPERTY(meta = (DataflowInput, DataflowOutput, DisplayName = "Collection", DataflowPassthrough = "Collection"))
+//		FManagedArrayCollection Collection;
+//
+//	UPROPERTY(meta = (DataflowInput, DisplayName = "TargetIndicesIn"))
+//		TArray<int32> TargetIndicesIn;
+//
+//	UPROPERTY(EditAnywhere, Category = "Dataflow")
+//		float PositionTargetStiffness = 10000.f;
+//
+//	UPROPERTY(meta = (DataflowOutput, DisplayName = "VertexBVH"))
+//		TArray<int32> VertexBVHOut;
+//
+//
+//
+//	FComputeVertexSphereBVHDataflowNode(const Dataflow::FNodeParameters& InParam, FGuid InGuid = FGuid::NewGuid())
+//		: FDataflowNode(InParam, InGuid)
+//	{
+//		RegisterInputConnection(&Collection);
+//		RegisterOutputConnection(&Collection, &Collection);
+//		RegisterInputConnection(&TargetIndicesIn);
+//		RegisterOutputConnection(&VertexBVHOut);
+//	}
+//
+//	virtual void Evaluate(Dataflow::FContext& Context, const FDataflowOutput* Out) const override;
+//};
+
+
+USTRUCT(meta = (DataflowFlesh))
+struct FSetVertexTetrahedraPositionTargetBindingDataflowNode : public FDataflowNode
+{
+	GENERATED_USTRUCT_BODY()
+		DATAFLOW_NODE_DEFINE_INTERNAL(FSetVertexTetrahedraPositionTargetBindingDataflowNode, "FSetVertexTetrahedraPositionTargetBinding", "Flesh", "")
+		DATAFLOW_NODE_RENDER_TYPE(FGeometryCollection::StaticType(), "Collection")
+
+public:
+	typedef FManagedArrayCollection DataType;
+
+
+	UPROPERTY(meta = (DataflowInput, DataflowOutput, DisplayName = "Collection", DataflowPassthrough = "Collection"))
+		FManagedArrayCollection Collection;
+
+	UPROPERTY(meta = (DataflowInput, DisplayName = "TargetIndicesIn"))
+		TArray<int32> TargetIndicesIn;
+
+	UPROPERTY(EditAnywhere, Category = "Dataflow")
+		float PositionTargetStiffness = 10000.f;
+
+
+
+	FSetVertexTetrahedraPositionTargetBindingDataflowNode(const Dataflow::FNodeParameters& InParam, FGuid InGuid = FGuid::NewGuid())
+		: FDataflowNode(InParam, InGuid)
+	{
+		RegisterInputConnection(&Collection);
+		RegisterOutputConnection(&Collection, &Collection);
+		RegisterInputConnection(&TargetIndicesIn);
+	}
+
+	virtual void Evaluate(Dataflow::FContext& Context, const FDataflowOutput* Out) const override;
+};
+
+namespace Dataflow
+{
+	void RegisterChaosFleshPositionTargetInitializationNodes();
+}
diff --git a/Engine/Source/Runtime/Experimental/Chaos/Private/Chaos/Deformable/ChaosDeformableSolver.cpp b/Engine/Source/Runtime/Experimental/Chaos/Private/Chaos/Deformable/ChaosDeformableSolver.cpp
--- a/Engine/Source/Runtime/Experimental/Chaos/Private/Chaos/Deformable/ChaosDeformableSolver.cpp
+++ b/Engine/Source/Runtime/Experimental/Chaos/Private/Chaos/Deformable/ChaosDeformableSolver.cpp
@@ -27,6 +27,7 @@
 #include "CoreMinimal.h"
 #include "GeometryCollection/Facades/CollectionKinematicBindingFacade.h"
 #include "GeometryCollection/Facades/CollectionVertexBoneWeightsFacade.h"
+#include "GeometryCollection/Facades/CollectionPositionTargetFacade.h"
 #include "Misc/FileHelper.h"
 #include "HAL/PlatformFileManager.h"
 #include "HAL/IConsoleManager.h"
@@ -50,6 +51,7 @@ DECLARE_CYCLE_STAT(TEXT("Chaos.Deformable.Solver.InitializeDeformableParticles")
 DECLARE_CYCLE_STAT(TEXT("Chaos.Deformable.Solver.InitializeKinematicParticles"), STAT_ChaosDeformableSolver_InitializeKinematicParticles, STATGROUP_Chaos);
 DECLARE_CYCLE_STAT(TEXT("Chaos.Deformable.Solver.InitializeTetrahedralConstraint"), STAT_ChaosDeformableSolver_InitializeTetrahedralConstraint, STATGROUP_Chaos);
 DECLARE_CYCLE_STAT(TEXT("Chaos.Deformable.Solver.InitializeGidBasedConstraints"), STAT_ChaosDeformableSolver_InitializeGidBasedConstraints, STATGROUP_Chaos);
+DECLARE_CYCLE_STAT(TEXT("Chaos.Deformable.Solver.InitializeWeakConstraints"), STAT_ChaosDeformableSolver_InitializeWeakConstraints, STATGROUP_Chaos);
 DECLARE_CYCLE_STAT(TEXT("Chaos.Deformable.Solver.InitializeKinematicConstraint"), STAT_ChaosDeformableSolver_InitializeKinematicConstraint, STATGROUP_Chaos);
 DECLARE_CYCLE_STAT(TEXT("Chaos.Deformable.Solver.InitializeCollisionBodies"), STAT_ChaosDeformableSolver_InitializeCollisionBodies, STATGROUP_Chaos);
 DECLARE_CYCLE_STAT(TEXT("Chaos.Deformable.Solver.InitializeSelfCollisionVariables"), STAT_ChaosDeformableSolver_InitializeSelfCollisionVariables, STATGROUP_Chaos);
@@ -108,7 +110,7 @@ namespace Chaos::Softs
 		MObjects = TArrayCollectionArray<const UObject*>();
 		FSolverParticles LocalParticlesDummy;
 		FSolverRigidParticles RigidParticles;
-		Evolution.Reset(new FPBDEvolution(MoveTemp(LocalParticlesDummy), MoveTemp(RigidParticles), {}, 
+		Evolution.Reset(new FPBDEvolution(MoveTemp(LocalParticlesDummy), MoveTemp(RigidParticles), {},
 			Property.NumSolverIterations, (FSolverReal)0.,
 			/*SelfCollisionsThickness = */(FSolverReal)0.,
 			/*CoefficientOfFriction = */(FSolverReal)0.,
@@ -194,6 +196,7 @@ namespace Chaos::Softs
 		{
 			InitializeDeformableParticles(*Proxy);
 			InitializeKinematicParticles(*Proxy);
+			InitializeWeakConstraint(*Proxy);
 			InitializeTetrahedralConstraint(*Proxy);
 			InitializeGidBasedConstraints(*Proxy);
 		}
@@ -217,7 +220,7 @@ namespace Chaos::Softs
 		auto ChaosInvM = [](FSolverReal M) { return FSolverReal(FMath::IsNearlyZero(M) ? 0.0 : 1 / M); };
 		auto DoubleVert = [](FVector3f V) { return FVector3d(V.X, V.Y, V.Z); };
 		uint32 NumParticles = Rest.NumElements(FGeometryCollection::VerticesGroup);
-		
+
 		int32 ParticleStart = Evolution->AddParticleRange(NumParticles, GroupOffset, true);
 		GroupOffset += 1;
 		for (uint32 vdx = 0; vdx < NumParticles; ++vdx)
@@ -254,10 +257,10 @@ namespace Chaos::Softs
 		}
 
 		Evolution->SetDamping(DampingMultiplier, GroupOffset - 1);
-		
+
 		const FTransform& InitialTransform = Proxy.GetInitialTransform();
 
-		
+
 		for (uint32 vdx = 0; vdx < NumParticles; ++vdx)
 		{
 			int32 SolverParticleIndex = ParticleStart + vdx;
@@ -323,8 +326,73 @@ namespace Chaos::Softs
 				}
 			}
 		}
+
+		
 	}
 
+	void FDeformableSolver::InitializeWeakConstraint(FFleshThreadingProxy& Proxy)
+	{
+		SCOPE_CYCLE_COUNTER(STAT_ChaosDeformableSolver_InitializeWeakConstraints);
+		TRACE_CPUPROFILER_EVENT_SCOPE(ChaosDeformableSolver_InitializeWeakConstraints);
+		const FManagedArrayCollection& Rest = Proxy.GetRestCollection();
+		const FIntVector2& Range = Proxy.GetSolverParticleRange();
+		if (Property.bEnablePositionTargets)
+		{
+			typedef GeometryCollection::Facades::FPositionTargetFacade FPositionTargets;
+			FPositionTargets PositionTargets(Rest);
+
+			TSet<int32> ConstrainedVertices;
+
+			TArray<TArray<int32>> PositionTargetIndices;
+			TArray<TArray<FSolverReal>> PositionTargetWeights;
+			TArray<TArray<int32>> PositionTargetSecondIndices;
+			TArray<TArray<FSolverReal>> PositionTargetSecondWeights;
+			TArray<FSolverReal> PositionTargetStiffness;
+
+			PositionTargetIndices.SetNum(PositionTargets.NumPositionTargets());
+			PositionTargetWeights.SetNum(PositionTargets.NumPositionTargets());
+			PositionTargetSecondIndices.SetNum(PositionTargets.NumPositionTargets());
+			PositionTargetSecondWeights.SetNum(PositionTargets.NumPositionTargets());
+			PositionTargetStiffness.SetNum(PositionTargets.NumPositionTargets());
+
+			// Read in position target info
+			for (int i = PositionTargets.NumPositionTargets() - 1; i >= 0; i--)
+			{
+				GeometryCollection::Facades::FPositionTargetsData DataPackage = PositionTargets.GetPositionTarget(i);
+				PositionTargetIndices[i] = DataPackage.SourceIndex;
+				PositionTargetWeights[i] = DataPackage.SourceWeights;
+				PositionTargetSecondIndices[i] = DataPackage.TargetIndex;
+				PositionTargetSecondWeights[i] = DataPackage.TargetWeights;
+				PositionTargetStiffness[i] = DataPackage.Stiffness;
+			}
+
+
+			int32 InitIndex = Evolution->AddConstraintInitRange(1, true);
+			int32 ConstraintIndex = Evolution->AddConstraintRuleRange(1, true);
+
+			FXPBDWeakConstraints<FSolverReal, FSolverParticles>* WeakConstraint =
+				new FXPBDWeakConstraints<FSolverReal, FSolverParticles>(Evolution->Particles(),
+					PositionTargetIndices, PositionTargetWeights, PositionTargetStiffness, PositionTargetSecondIndices, PositionTargetSecondWeights);
+
+			Evolution->ConstraintInits()[InitIndex] =
+				[WeakConstraint, this](FSolverParticles& InParticles, const FSolverReal Dt)
+			{
+				WeakConstraint->Init();
+				//WeakConstraint->UpdateTargets(this->ComputeParticleTargets(WeakConstraint->GetIndices()));
+			};
+
+			Evolution->ConstraintRules()[ConstraintIndex] =
+				[WeakConstraint](FSolverParticles& InParticles, const FSolverReal Dt)
+			{
+				WeakConstraint->ApplyInParallel(InParticles, Dt);
+			};
+
+			WeakConstraints.Add(TUniquePtr<FXPBDWeakConstraints<FSolverReal, FSolverParticles>>(WeakConstraint));
+
+		}
+	}
+
+
 	void FDeformableSolver::DebugDrawTetrahedralParticles(FFleshThreadingProxy& Proxy)
 	{
 #if WITH_EDITOR
@@ -416,7 +484,7 @@ namespace Chaos::Softs
 					}
 
 				}
-				
+
 				if (Property.bEnableCorotatedConstraints)
 				{
 
@@ -466,7 +534,7 @@ namespace Chaos::Softs
 					}
 				}
 			}
-		
+
 		}
 	}
 
@@ -518,7 +586,6 @@ namespace Chaos::Softs
 						if (Rest.FindAttributeTyped<FVector3f>("Vertex", FGeometryCollection::VerticesGroup))
 						{
 							const TManagedArray<FVector3f>& Vertex = Rest.GetAttribute<FVector3f>("Vertex", FGeometryCollection::VerticesGroup);
-
 							// @todo(chaos) : reduce conversions
 							auto ChaosVert = [](FVector3f V) { return Chaos::FVec3(V.X, V.Y, V.Z); };
 							auto ChaosVertfloat = [](FVector3f V) { return Chaos::TVector<FSolverReal, 3>(V.X, V.Y, V.Z); };
@@ -558,39 +625,50 @@ namespace Chaos::Softs
 										TVector<FSolverReal, 3> TargetPos((FSolverReal)0.);
 										FSolverReal CurrentRatio = FSolverReal(this->Iteration) / FSolverReal(this->Property.NumSolverSubSteps);
 
-										const TManagedArray<FTransform>& RestTransforms = Rest.GetAttribute<FTransform>(FTransformCollection::TransformAttribute, FTransformCollection::TransformGroup);
-
-										for (int32 i = BoneIndices.Num() - 1; 0 <= i; i--)
+										int32 RestNum = FleshInputBuffer->RestTransforms.Num();
+										int32 TransformNum = FleshInputBuffer->Transforms.Num();
+										if (RestNum > 0 && TransformNum > 0)
 										{
-											// @todo(flesh) : Add the pre-cached component space rest transforms to the rest collection. 
-											// see  UFleshComponent::NewDeformableData for how its pulled from the SkeletalMesh
-											FVec3 LocalPoint = FleshInputBuffer->RestTransforms[BoneIndices[i]].InverseTransformPosition(ChaosVert(Vertex[Index - Range[0]]));
-											FVec3 ComponentPointAtT = FleshInputBuffer->Transforms[BoneIndices[i]].TransformPosition(LocalPoint);
-											if (PreviousFleshBuffer)
+											
+											for (int32 i = BoneIndices.Num() - 1; 0 <= i; i--)
 											{
-												FTransform BonePreviousTransform = PreviousFleshBuffer->Transforms[BoneIndices[i]];
-												ComponentPointAtT = ComponentPointAtT * CurrentRatio + BonePreviousTransform.TransformPosition(LocalPoint) * ((FSolverReal)1. - CurrentRatio);
-											}
+												if (BoneIndices[i] > -1 && BoneIndices[i] < RestNum && BoneIndices[i] < TransformNum)
+												{
 
-											MParticles.X(Index) = GlobalTransform.TransformPosition(ComponentPointAtT);
-											bParticleTouched = true;
+													// @todo(flesh) : Add the pre-cached component space rest transforms to the rest collection. 
+													// see  UFleshComponent::NewDeformableData for how its pulled from the SkeletalMesh
+													FVec3 LocalPoint = FleshInputBuffer->RestTransforms[BoneIndices[i]].InverseTransformPosition(ChaosVert(Vertex[Index - Range[0]]));
+													FVec3 ComponentPointAtT = FleshInputBuffer->Transforms[BoneIndices[i]].TransformPosition(LocalPoint);
 
-											// @todo(flesh): Add non rigid skinning weights. 
-											// Currently this just grabs the last joint, ideally
-											// we would be using the weighted average for the vertex
-											// positions.
+													if (PreviousFleshBuffer)
+													{
+														FTransform BonePreviousTransform = PreviousFleshBuffer->Transforms[BoneIndices[i]];
+														ComponentPointAtT = ComponentPointAtT * CurrentRatio + BonePreviousTransform.TransformPosition(LocalPoint) * ((FSolverReal)1. - CurrentRatio);
+													}
+
+													MParticles.X(Index) = GlobalTransform.TransformPosition(ComponentPointAtT);
+
+													bParticleTouched = true;
+
+
+													// @todo(flesh): Add non rigid skinning weights. 
+													// Currently this just grabs the last joint, ideally
+													// we would be using the weighted average for the vertex
+													// positions.
 
 #if WITH_EDITOR
-											//debug draw
-											//p.Chaos.DebugDraw.Enabled 1
-											//p.Chaos.DebugDraw.Deformable.KinematicParticle 1
-											if (GDeformableDebugParams.IsDebugDrawingEnabled() && GDeformableDebugParams.bDoDrawKinematicParticles)
-											{
-												auto DoubleVert = [](FVector3f V) { return FVector3d(V.X, V.Y, V.Z); };
-												Chaos::FDebugDrawQueue::GetInstance().DrawDebugPoint(DoubleVert(MParticles.X(Index)), FColor::Red, false, -1.0f, 0, 5);
-											}
+													//debug draw
+													//p.Chaos.DebugDraw.Enabled 1
+													//p.Chaos.DebugDraw.Deformable.KinematicParticle 1
+													if (GDeformableDebugParams.IsDebugDrawingEnabled() && GDeformableDebugParams.bDoDrawKinematicParticles)
+													{
+														auto DoubleVert = [](FVector3f V) { return FVector3d(V.X, V.Y, V.Z); };
+														Chaos::FDebugDrawQueue::GetInstance().DrawDebugPoint(DoubleVert(MParticles.X(Index)), FColor::Red, false, -1.0f, 0, 5);
+													}
 #endif
-											break;
+													break;
+												}
+											}
 										}
 										MParticles.PAndInvM(Index).P = MParticles.X(Index);
 									}
@@ -600,6 +678,17 @@ namespace Chaos::Softs
 							{
 								MParticles.X(Index) = GlobalTransform.TransformPosition(ChaosVert(Vertex[Index - Range[0]]));
 								MParticles.PAndInvM(Index).P = MParticles.X(Index);
+
+#if WITH_EDITOR
+								//debug draw
+								//p.Chaos.DebugDraw.Enabled 1
+								//p.Chaos.DebugDraw.Deformable.KinematicParticle 1
+								if (GDeformableDebugParams.IsDebugDrawingEnabled() && GDeformableDebugParams.bDoDrawKinematicParticles)
+								{
+									auto DoubleVert = [](FVector3f V) { return FVector3d(V.X, V.Y, V.Z); };
+									Chaos::FDebugDrawQueue::GetInstance().DrawDebugPoint(DoubleVert(MParticles.X(Index)), FColor::Red, false, -1.0f, 0, 5);
+								}
+#endif
 							}
 						}
 					}
diff --git a/Engine/Source/Runtime/Experimental/Chaos/Private/Chaos/GraphColoring.cpp b/Engine/Source/Runtime/Experimental/Chaos/Private/Chaos/GraphColoring.cpp
--- a/Engine/Source/Runtime/Experimental/Chaos/Private/Chaos/GraphColoring.cpp
+++ b/Engine/Source/Runtime/Experimental/Chaos/Private/Chaos/GraphColoring.cpp
@@ -205,6 +205,60 @@ static bool VerifyGridBasedSubColoring(const TArray<TArray<int32>>& ElementsPerC
 
 }
 
+template <typename T>
+static bool VerifyWeakConstraintsColoring(const TArray<TArray<int32>>& Indices, const TArray<TArray<int32>>& SecondIndices, const Chaos::TDynamicParticles<T, 3>& InParticles, const TArray<TArray<int32>>& ConstraintsPerColor)
+{
+	TArray<bool> ConstraintIsIncluded;
+	ConstraintIsIncluded.Init(false, Indices.Num());
+	for (int32 i = 0; i < ConstraintsPerColor.Num(); i++)
+	{
+		for (int32 j = 0; j < ConstraintsPerColor[i].Num(); j++)
+		{
+			ConstraintIsIncluded[ConstraintsPerColor[i][j]] = true;
+		}
+	}
+
+	for (int32 kk = 0; kk < Indices.Num(); kk++)
+	{
+		if (!ConstraintIsIncluded[kk])
+		{
+			return false;
+		}
+	}
+
+	for (int32 i = 0; i < ConstraintsPerColor.Num(); i++)
+	{
+		TSet<int32> CoveredParticles;
+		for (int32 j = 0; j < ConstraintsPerColor[i].Num(); j++)
+		{
+			for (int32 Node : Indices[ConstraintsPerColor[i][j]])
+			{
+				if (CoveredParticles.Contains(Node))
+				{
+					return false;
+				}
+				CoveredParticles.Emplace(Node);
+			}
+
+			if (SecondIndices.Num() > 0)
+			{
+				for (int32 Node : SecondIndices[ConstraintsPerColor[i][j]])
+				{
+					if (CoveredParticles.Contains(Node))
+					{
+						return false;
+					}
+					CoveredParticles.Emplace(Node);
+				}
+
+			}
+		}
+	}
+
+	return true;
+
+}
+
 template<typename T>
 TArray<TArray<int32>> Chaos::FGraphColoring::ComputeGraphColoring(const TArray<Chaos::TVec2<int32>>& Graph, const Chaos::TDynamicParticles<T, 3>& InParticles, const int32 GraphParticlesStart, const int32 GraphParticlesEnd)
 {
@@ -1050,7 +1104,6 @@ void Chaos::ComputeGridBasedGraphSubColoringPointer(const TArray<TArray<int32>>&
 			ElementSubColors.Init(-1, ElementsPerColor[i].Num());
 			TArray<TSet<int32>*> UsedColors;
 			UsedColors.Init(nullptr, NumNodes);
-			//std::vector<std::unordered_set<int>*> UsedColors(NumNodes, nullptr);
 			for (int32 j = 0; j < ElementsPerColor[i].Num(); j++) {
 				int32 ColorToUse = 0;
 				int32 e = ElementsPerColor[i][j];
@@ -1123,6 +1176,122 @@ void Chaos::ComputeGridBasedGraphSubColoringPointer(const TArray<TArray<int32>>&
 }
 
 
+template<typename T>
+void Chaos::ComputeWeakConstraintsColoring(const TArray<TArray<int32>>& Indices, const TArray<TArray<int32>>& SecondIndices, const Chaos::TDynamicParticles<T, 3>& InParticles, TArray<TArray<int32>>& ConstraintsPerColor)
+{
+	TArray<TSet<int32>*> UsedColors;
+	UsedColors.Init(nullptr, InParticles.Size());
+
+	ensure(Indices.Num() == SecondIndices.Num() || SecondIndices.Num() == 0);
+
+	TArray<int32> ConstraintColors;
+	ConstraintColors.Init(-1, Indices.Num());
+
+	if (SecondIndices.Num() == 0)
+	{
+		for (int32 ConstraintIndex = 0; ConstraintIndex < Indices.Num(); ConstraintIndex++)
+		{
+			for (int32 i = 0; i < Indices[ConstraintIndex].Num(); i++)
+			{
+				if (!UsedColors[Indices[ConstraintIndex][i]]) {
+					UsedColors[Indices[ConstraintIndex][i]] = new TSet<int32>();
+				}
+				
+			}
+			int32 ColorToUse = 0;
+			while (true)
+			{
+				bool ColorFound = false;
+				for (auto node : Indices[ConstraintIndex]) {
+					if (!UsedColors[node]) {
+						UsedColors[node] = new TSet<int32>();
+					}
+					if (UsedColors[node]->Contains(ColorToUse)) {
+						ColorFound = true;
+						break;
+					}
+				}
+				if (!ColorFound) {
+					break;
+				}
+				ColorToUse++;
+			}
+			ConstraintColors[ConstraintIndex] = ColorToUse;
+		}
+	}
+	else 
+	{
+		for (int32 ConstraintIndex = 0; ConstraintIndex < Indices.Num(); ConstraintIndex++)
+		{
+			for (int32 i = 0; i < Indices[ConstraintIndex].Num(); i++)
+			{
+				if (!UsedColors[Indices[ConstraintIndex][i]]) {
+					UsedColors[Indices[ConstraintIndex][i]] = new TSet<int32>();
+				}
+			}
+			for (int32 j = 0; j < SecondIndices[ConstraintIndex].Num(); j++)
+			{
+				if (!UsedColors[SecondIndices[ConstraintIndex][j]]) {
+					UsedColors[SecondIndices[ConstraintIndex][j]] = new TSet<int32>();
+				}
+			}
+			int32 ColorToUse = 0;
+			while (true)
+			{
+				bool ColorFound = false;
+				for (auto node : Indices[ConstraintIndex]) {
+					if (!UsedColors[node]) {
+						UsedColors[node] = new TSet<int32>();
+					}
+					if (UsedColors[node]->Contains(ColorToUse)) {
+						ColorFound = true;
+						break;
+					}
+				}
+				for (auto node : SecondIndices[ConstraintIndex]) {
+					if (!UsedColors[node]) {
+						UsedColors[node] = new TSet<int32>();
+					}
+					if (UsedColors[node]->Contains(ColorToUse)) {
+						ColorFound = true;
+						break;
+					}
+				}
+				if (!ColorFound) {
+					break;
+				}
+				ColorToUse++;
+			}
+			ConstraintColors[ConstraintIndex] = ColorToUse;
+			for (auto node : Indices[ConstraintIndex]) {
+				UsedColors[node]->Emplace(ColorToUse);
+			}
+			for (auto node : SecondIndices[ConstraintIndex]) {
+				UsedColors[node]->Emplace(ColorToUse);
+			}
+		}
+	}
+
+	for (int ii = 0; ii < UsedColors.Num(); ii++) {
+		delete (UsedColors[ii]);
+	}
+
+	int32 NumColors = FMath::Max<int32>(ConstraintColors);
+
+	//int32 num_colors = *std::max_element(ElementSubColors.begin(), ElementSubColors.end());
+	ConstraintsPerColor.Empty();
+	ConstraintsPerColor.SetNum(NumColors + 1);
+
+	for (int32 j = 0; j < Indices.Num(); j++) 
+	{
+		ConstraintsPerColor[ConstraintColors[j]].Emplace(j);
+	}
+
+	checkSlow(VerifyWeakConstraintsColoring<T>(Indices, SecondIndices, InParticles, ConstraintsPerColor));
+}
+
+
+
 template CHAOS_API TArray<TArray<int32>> Chaos::FGraphColoring::ComputeGraphColoring<Chaos::FRealSingle>(const TArray<Chaos::TVector<int32, 2>>&, const Chaos::TDynamicParticles<Chaos::FRealSingle, 3>&, const int32 GraphParticlesStart, const int32 GraphParticlesEnd);
 template CHAOS_API TArray<TArray<int32>> Chaos::FGraphColoring::ComputeGraphColoring<Chaos::FRealDouble>(const TArray<Chaos::TVector<int32, 2>>&, const Chaos::TDynamicParticles<Chaos::FRealDouble, 3>&, const int32 GraphParticlesStart, const int32 GraphParticlesEnd);
 template CHAOS_API TArray<TArray<int32>> Chaos::FGraphColoring::ComputeGraphColoring<Chaos::FRealSingle>(const TArray<Chaos::TVector<int32, 3>>&, const Chaos::TDynamicParticles<Chaos::FRealSingle, 3>&, const int32 GraphParticlesStart, const int32 GraphParticlesEnd);
@@ -1132,4 +1301,6 @@ template CHAOS_API TArray<TArray<int32>> Chaos::FGraphColoring::ComputeGraphColo
 template CHAOS_API TArray<TArray<int32>> Chaos::FGraphColoring::ComputeGraphColoringAllDynamic<Chaos::FRealSingle>(const TArray<Chaos::TVector<int32, 4>>&, const Chaos::TDynamicParticles<Chaos::FRealSingle, 3>&, const int32 GraphParticlesStart, const int32 GraphParticlesEnd);
 template CHAOS_API TArray<TArray<int32>> Chaos::FGraphColoring::ComputeGraphColoringAllDynamic<Chaos::FRealDouble>(const TArray<Chaos::TVector<int32, 4>>&, const Chaos::TDynamicParticles<Chaos::FRealDouble, 3>&, const int32 GraphParticlesStart, const int32 GraphParticlesEnd);
 template CHAOS_API void Chaos::ComputeGridBasedGraphSubColoringPointer(const TArray<TArray<int32>>& ElementsPerColor, const TMPMGrid<Chaos::FRealSingle>& Grid, const int32 GridSize, TArray<TArray<int32>>*& PreviousColoring, const TArray<TArray<int32>>& ConstraintsNodesSet, TArray<TArray<TArray<int32>>>& ElementsPerSubColors);
-template CHAOS_API void Chaos::ComputeGridBasedGraphSubColoringPointer(const TArray<TArray<int32>>& ElementsPerColor, const TMPMGrid<Chaos::FRealDouble>& Grid, const int32 GridSize, TArray<TArray<int32>>*& PreviousColoring, const TArray<TArray<int32>>& ConstraintsNodesSet, TArray<TArray<TArray<int32>>>& ElementsPerSubColors);
\ No newline at end of file
+template CHAOS_API void Chaos::ComputeGridBasedGraphSubColoringPointer(const TArray<TArray<int32>>& ElementsPerColor, const TMPMGrid<Chaos::FRealDouble>& Grid, const int32 GridSize, TArray<TArray<int32>>*& PreviousColoring, const TArray<TArray<int32>>& ConstraintsNodesSet, TArray<TArray<TArray<int32>>>& ElementsPerSubColors);
+template CHAOS_API void Chaos::ComputeWeakConstraintsColoring<Chaos::FRealSingle>(const TArray<TArray<int32>>& Indices, const TArray<TArray<int32>>& SecondIndices, const Chaos::TDynamicParticles<Chaos::FRealSingle, 3>& InParticles, TArray<TArray<int32>>& ConstraintsPerColor);
+template CHAOS_API void Chaos::ComputeWeakConstraintsColoring<Chaos::FRealDouble>(const TArray<TArray<int32>>& Indices, const TArray<TArray<int32>>& SecondIndices, const Chaos::TDynamicParticles<Chaos::FRealDouble, 3>& InParticles, TArray<TArray<int32>>& ConstraintsPerColor);
\ No newline at end of file
diff --git a/Engine/Source/Runtime/Experimental/Chaos/Private/GeometryCollection/Facades/CollectionPositionTargetFacade.cpp b/Engine/Source/Runtime/Experimental/Chaos/Private/GeometryCollection/Facades/CollectionPositionTargetFacade.cpp
new file mode 100644
--- /dev/null
+++ b/Engine/Source/Runtime/Experimental/Chaos/Private/GeometryCollection/Facades/CollectionPositionTargetFacade.cpp
@@ -0,0 +1,160 @@
+// Copyright Epic Games, Inc. All Rights Reserved.
+
+/*=============================================================================
+	FleshCollection.cpp: FFleshCollection methods.
+=============================================================================*/
+
+#include "GeometryCollection/Facades/CollectionPositionTargetFacade.h"
+#include "GeometryCollection/GeometryCollection.h"
+
+namespace GeometryCollection::Facades
+{
+	// Attributes
+	const FName FPositionTargetFacade::GroupName("PositionTargets");
+	const FName FPositionTargetFacade::TargetIndex("TargetIndex");
+	const FName FPositionTargetFacade::SourceIndex("SourceIndex");
+	const FName FPositionTargetFacade::Stiffness("Stiffness");
+	const FName FPositionTargetFacade::Damping("Damping");
+	const FName FPositionTargetFacade::SourceName("SourceName");
+	const FName FPositionTargetFacade::TargetName("TargetName");
+	const FName FPositionTargetFacade::TargetWeights("TargetWeights");
+	const FName FPositionTargetFacade::SourceWeights("SourceWeights");
+
+	FPositionTargetFacade::FPositionTargetFacade(FManagedArrayCollection& InCollection)
+		: ConstCollection(InCollection)
+		, Collection(&InCollection)
+		, TargetIndexAttribute(InCollection, TargetIndex, GroupName)
+		, SourceIndexAttribute(InCollection, SourceIndex, GroupName)
+		, TargetNameAttribute(InCollection, TargetName, GroupName)
+		, SourceNameAttribute(InCollection, SourceName, GroupName)
+		, StiffnessAttribute(InCollection, Stiffness, GroupName)
+		, DampingAttribute(InCollection, Damping, GroupName)
+		, TargetWeightsAttribute(InCollection, TargetWeights, GroupName)
+		, SourceWeightsAttribute(InCollection, SourceWeights, GroupName)
+	{
+		DefineSchema();
+	}
+
+	FPositionTargetFacade::FPositionTargetFacade(const FManagedArrayCollection& InCollection)
+		: ConstCollection(InCollection)
+		, Collection(nullptr)
+		, TargetIndexAttribute(InCollection, TargetIndex, GroupName)
+		, SourceIndexAttribute(InCollection, SourceIndex, GroupName)
+		, TargetNameAttribute(InCollection, TargetName, GroupName)
+		, SourceNameAttribute(InCollection, SourceName, GroupName)
+		, StiffnessAttribute(InCollection, Stiffness, GroupName)
+		, DampingAttribute(InCollection, Damping, GroupName)
+		, TargetWeightsAttribute(InCollection, TargetWeights, GroupName)
+		, SourceWeightsAttribute(InCollection, SourceWeights, GroupName)
+	{
+		//DefineSchema();
+	}
+
+	bool FPositionTargetFacade::IsValid() const
+	{
+		return TargetIndexAttribute.IsValid() && SourceIndexAttribute.IsValid()
+			&& TargetNameAttribute.IsValid() && SourceNameAttribute.IsValid() && StiffnessAttribute.IsValid()
+			&& DampingAttribute.IsValid() && TargetWeightsAttribute.IsValid() && SourceWeightsAttribute.IsValid();
+	}
+
+	void FPositionTargetFacade::DefineSchema()
+	{
+		check(!IsConst());
+		TargetIndexAttribute.Add(ManageArrayAccessor::EPersistencePolicy::MakePersistent, FGeometryCollection::VerticesGroup);
+		SourceIndexAttribute.Add(ManageArrayAccessor::EPersistencePolicy::MakePersistent, FGeometryCollection::VerticesGroup);
+		TargetNameAttribute.Add();
+		SourceNameAttribute.Add();
+		StiffnessAttribute.Add();
+		DampingAttribute.Add();
+		TargetWeightsAttribute.Add();
+		SourceWeightsAttribute.Add();
+	}
+
+
+	int32 FPositionTargetFacade::AddPositionTarget(const FPositionTargetsData& InputData)
+	{
+		check(!IsConst());
+		if (IsValid())
+		{
+			int32 NewIndex = TargetIndexAttribute.AddElements(1);
+			//TManagedArray<TArray<int32>>& TargetIndex = TargetIndexAttribute.Modify();
+			TargetIndexAttribute.Modify()[NewIndex] = InputData.TargetIndex;
+			SourceIndexAttribute.Modify()[NewIndex] = InputData.SourceIndex;
+			TargetNameAttribute.Modify()[NewIndex] = InputData.TargetName;
+			SourceNameAttribute.Modify()[NewIndex] = InputData.SourceName;
+			StiffnessAttribute.Modify()[NewIndex] = InputData.Stiffness;
+			DampingAttribute.Modify()[NewIndex] = InputData.Damping;
+			TargetWeightsAttribute.Modify()[NewIndex] = InputData.TargetWeights;
+			SourceWeightsAttribute.Modify()[NewIndex] = InputData.SourceWeights;
+			// 
+			// 
+			//TManagedArray<FString>& PositionTargetToGroupLocal = PositionTargetToGroupAttribute.Modify();
+
+			//PositionTarget[NewIndex] = Key.Index;
+			//PositionTargetToGroupLocal[NewIndex] = Key.GroupName.ToString();
+
+			//TManagedArrayAccessor<TArray<int32>> TargetIndexAttribute;
+			//TManagedArrayAccessor<TArray<int32>> SourceIndexAttribute;
+			//TManagedArrayAccessor<FString> TargetNameAttribute;
+			//TManagedArrayAccessor<FString> SourceNameAttribute;
+			//TManagedArrayAccessor<float> StiffnessAttribute;
+			//TManagedArrayAccessor<float> DampingAttribute;
+
+			return NewIndex;
+		}
+		return INDEX_NONE;
+	}
+
+	FPositionTargetsData FPositionTargetFacade::GetPositionTarget(const int32 DataIndex) const
+	{
+		FPositionTargetsData ReturnData;
+		if (IsValid())
+		{
+			if (StiffnessAttribute.Num() > DataIndex && DataIndex > -1) 
+			{
+				ReturnData.Stiffness = StiffnessAttribute.Get()[DataIndex];
+			}
+
+			if (DampingAttribute.Num() > DataIndex && DataIndex > -1)
+			{
+				ReturnData.Damping = DampingAttribute.Get()[DataIndex];
+			}
+
+			if (SourceNameAttribute.Num() > DataIndex && DataIndex > -1)
+			{
+				ReturnData.SourceName = SourceNameAttribute.Get()[DataIndex];
+			}
+
+			if (TargetNameAttribute.Num() > DataIndex && DataIndex > -1)
+			{
+				ReturnData.TargetName = TargetNameAttribute.Get()[DataIndex];
+			}
+
+			if (SourceIndexAttribute.Num() > DataIndex && DataIndex > -1)
+			{
+				ReturnData.SourceIndex = SourceIndexAttribute.Get()[DataIndex];
+			}
+
+			if (TargetIndexAttribute.Num() > DataIndex && DataIndex > -1)
+			{
+				ReturnData.TargetIndex = TargetIndexAttribute.Get()[DataIndex];
+			}
+
+			if (SourceWeightsAttribute.Num() > DataIndex && DataIndex > -1)
+			{
+				ReturnData.SourceWeights = SourceWeightsAttribute.Get()[DataIndex];
+			}
+
+			if (TargetWeightsAttribute.Num() > DataIndex && DataIndex > -1)
+			{
+				ReturnData.TargetWeights = TargetWeightsAttribute.Get()[DataIndex];
+			}
+
+		}
+		return ReturnData;
+	
+	}
+
+
+
+}
diff --git a/Engine/Source/Runtime/Experimental/Chaos/Public/Chaos/Deformable/ChaosDeformableSolver.h b/Engine/Source/Runtime/Experimental/Chaos/Public/Chaos/Deformable/ChaosDeformableSolver.h
--- a/Engine/Source/Runtime/Experimental/Chaos/Public/Chaos/Deformable/ChaosDeformableSolver.h
+++ b/Engine/Source/Runtime/Experimental/Chaos/Public/Chaos/Deformable/ChaosDeformableSolver.h
@@ -11,7 +11,7 @@
 #include "Chaos/PBDTriangleMeshCollisions.h"
 #include "Chaos/TriangleMesh.h"
 #include "Chaos/XPBDCorotatedConstraints.h"
-//#include "Chaos/XPBDWeakConstraints.h"
+#include "Chaos/XPBDWeakConstraints.h"
 #include "Chaos/BlendedXPBDCorotatedConstraints.h"
 #include "Chaos/XPBDGridBasedCorotatedConstraints.h"
 #include "Chaos/Deformable/ChaosDeformableSolverTypes.h"
@@ -120,12 +120,13 @@ namespace Chaos::Softs
 		void InitializeKinematicParticles(FFleshThreadingProxy&);
 		void InitializeTetrahedralConstraint(FFleshThreadingProxy&);
 		void InitializeGidBasedConstraints(FFleshThreadingProxy&);
-		//void InitializeWeakConstraints(FFleshThreadingProxy&);
+		void InitializeWeakConstraint(FFleshThreadingProxy&);
 		void InitializeKinematicConstraint();
 		void InitializeCollisionBodies();
 		void InitializeSelfCollisionVariables();
 		void InitializeGridBasedConstraintVariables();
 		void RemoveSimulationObjects();
+		TArray<Chaos::TVec3<FSolverReal>> ComputeParticleTargets(const TArray<TArray<int32>>& ParticleIndices);
 
 
 		/*IO Utility*/
@@ -171,6 +172,7 @@ namespace Chaos::Softs
 		// Simulation Variables
 		TUniquePtr<Softs::FPBDEvolution> Evolution;
 		TArray<TUniquePtr<Softs::FXPBDCorotatedConstraints<Softs::FSolverReal, Softs::FSolverParticles>>> CorotatedConstraints;
+		TArray<TUniquePtr<Softs::FXPBDWeakConstraints<Softs::FSolverReal, Softs::FSolverParticles>>> WeakConstraints;
 		TArray<TUniquePtr<Softs::FBlendedXPBDCorotatedConstraints<Softs::FSolverReal, Softs::FSolverParticles>>> BlendedCorotatedConstraints;
 		TUniquePtr<Softs::FXPBDGridBasedCorotatedConstraints<Softs::FSolverReal, Softs::FSolverParticles>> GridBasedCorotatedConstraint;
 		TUniquePtr<Softs::FPBDCollisionSpringConstraints> CollisionSpringConstraint;
diff --git a/Engine/Source/Runtime/Experimental/Chaos/Public/Chaos/Deformable/ChaosDeformableSolverTypes.h b/Engine/Source/Runtime/Experimental/Chaos/Public/Chaos/Deformable/ChaosDeformableSolverTypes.h
--- a/Engine/Source/Runtime/Experimental/Chaos/Public/Chaos/Deformable/ChaosDeformableSolverTypes.h
+++ b/Engine/Source/Runtime/Experimental/Chaos/Public/Chaos/Deformable/ChaosDeformableSolverTypes.h
@@ -31,7 +31,8 @@ namespace Chaos::Softs
 			FSolverReal InBlendedZeta = (FSolverReal).1,
 			FSolverReal InDamping = (FSolverReal)0,
 			bool InbEnableGravity = true, 
-			bool InbEnableCorotatedConstraints = true)
+			bool InbEnableCorotatedConstraints = true, 
+			bool InbEnablePositionTargets = true)
 			: NumSolverSubSteps(InNumSolverSubSteps)
 			, NumSolverIterations(InNumSolverIterations)
 			, FixTimeStep(InFixTimeStep)
@@ -49,6 +50,7 @@ namespace Chaos::Softs
 			, Damping(InDamping)
 			, bEnableGravity(InbEnableGravity)
 			, bEnableCorotatedConstraints(InbEnableCorotatedConstraints)
+			, bEnablePositionTargets(InbEnablePositionTargets)
 		{}
 
 		FDeformableSolverProperties(const FDeformableSolverProperties& InProp)
@@ -69,6 +71,7 @@ namespace Chaos::Softs
 			, Damping(InProp.Damping)
 			, bEnableGravity(InProp.bEnableGravity)
 			, bEnableCorotatedConstraints(InProp.bEnableCorotatedConstraints)
+			, bEnablePositionTargets(InProp.bEnablePositionTargets)
 		{}
 
 		int32 NumSolverSubSteps = 5;
@@ -88,6 +91,7 @@ namespace Chaos::Softs
 		FSolverReal Damping = (FSolverReal)0.;
 		bool bEnableGravity = true;
 		bool bEnableCorotatedConstraints = true;
+		bool bEnablePositionTargets = true;
 	};
 
 
diff --git a/Engine/Source/Runtime/Experimental/Chaos/Public/Chaos/GraphColoring.h b/Engine/Source/Runtime/Experimental/Chaos/Public/Chaos/GraphColoring.h
--- a/Engine/Source/Runtime/Experimental/Chaos/Public/Chaos/GraphColoring.h
+++ b/Engine/Source/Runtime/Experimental/Chaos/Public/Chaos/GraphColoring.h
@@ -94,4 +94,8 @@ class FGraphColoring
 template<typename T> 
 CHAOS_API void ComputeGridBasedGraphSubColoringPointer(const TArray<TArray<int32>>& ElementsPerColor, const TMPMGrid<T>& Grid, const int32 GridSize, TArray<TArray<int32>>*& PreviousColoring, const TArray<TArray<int32>>& ConstraintsNodesSet, TArray<TArray<TArray<int32>>>& ElementsPerSubColors);
 
+
+template<typename T>
+CHAOS_API void ComputeWeakConstraintsColoring(const TArray<TArray<int32>>& Indices, const TArray<TArray<int32>>& SecondIndices, const Chaos::TDynamicParticles<T, 3>& InParticles, TArray<TArray<int32>>& ConstraintsPerColor);
+
 }
diff --git a/Engine/Source/Runtime/Experimental/Chaos/Public/Chaos/XPBDWeakConstraints.h b/Engine/Source/Runtime/Experimental/Chaos/Public/Chaos/XPBDWeakConstraints.h
new file mode 100644
--- /dev/null
+++ b/Engine/Source/Runtime/Experimental/Chaos/Public/Chaos/XPBDWeakConstraints.h
@@ -0,0 +1,228 @@
+// Copyright Epic Games, Inc. All Rights Reserved.
+#pragma once
+
+#include "Chaos/PBDSpringConstraintsBase.h"
+#include "ChaosStats.h"
+#include "Chaos/ImplicitQRSVD.h"
+#include "Chaos/GraphColoring.h"
+#include "Chaos/MPMTransfer.h"
+#include "Chaos/Framework/Parallel.h"
+#include "Misc/FileHelper.h"
+#include "HAL/PlatformFileManager.h"
+#include "HAL/IConsoleManager.h"
+#include "Misc/Paths.h"
+#include "Chaos/GraphColoring.h"
+
+//DECLARE_CYCLE_STAT(TEXT("Chaos XPBD Corotated Constraint"), STAT_ChaosXPBDCorotated, STATGROUP_Chaos);
+//DECLARE_CYCLE_STAT(TEXT("Chaos XPBD Corotated Constraint Polar Compute"), STAT_ChaosXPBDCorotatedPolar, STATGROUP_Chaos);
+//DECLARE_CYCLE_STAT(TEXT("Chaos XPBD Corotated Constraint Det Compute"), STAT_ChaosXPBDCorotatedDet, STATGROUP_Chaos);
+
+namespace Chaos::Softs
+{
+
+	using Chaos::TVec3;
+
+	template <typename T, typename ParticleType>
+	class FXPBDWeakConstraints
+	{
+
+	public:
+		//this one only accepts tetmesh input and mesh
+		FXPBDWeakConstraints(
+			const ParticleType& InParticles,
+			const TArray<TArray<int32>>& InIndices,
+			const TArray<TArray<T>>& InWeights,
+			const TArray<T>& InStiffness
+		)
+			: Stiffness(InStiffness), Indices(InIndices), Weights(InWeights)
+		{
+			SecondIndices.SetNum(0);
+			SecondWeights.SetNum(0);
+			InitColor(InParticles);
+			LambdaArray.Init((T)0., Indices.Num() * 3);
+		}
+
+		FXPBDWeakConstraints(
+			const ParticleType& InParticles,
+			const TArray<TArray<int32>>& InIndices,
+			const TArray<TArray<T>>& InWeights,
+			const TArray<T>& InStiffness,
+			const TArray<TArray<int32>>& InSecondIndices,
+			const TArray<TArray<T>>& InSecondWeights
+		)
+			: Indices(InIndices), Weights(InWeights), SecondIndices(InSecondIndices), SecondWeights(InSecondWeights), Stiffness(InStiffness)
+		{
+			ensureMsgf(Indices.Num() == SecondIndices.Num(), TEXT("Input Double Bindings have wrong size"));
+			
+			for (int32 i = 0; i < Indices.Num(); i++)
+			{
+				TSet<int32> IndicesSet = TSet<int32>(Indices[i]);
+				for (int32 j = 0; j < SecondIndices[i].Num(); j++)
+				{
+					ensureMsgf(!IndicesSet.Contains(SecondIndices[i][j]), TEXT("Indices and Second Indices overlaps. Currently not supported"));
+				}
+			}
+			InitColor(InParticles);
+			LambdaArray.Init((T)0., Indices.Num() * 3);
+		}
+
+
+		virtual ~FXPBDWeakConstraints() {}
+
+		void ApplyInParallel(ParticleType& Particles, const T Dt) const
+		{
+			TRACE_CPUPROFILER_EVENT_SCOPE(TEXT("STAT_ChaosXPBDWeakConstraintApply"));
+
+			if ((ConstraintsPerColor.Num() > 0))
+			{
+				if (SecondIndices.Num() == 0)
+				{
+					for (int32 Color = 0; Color < ConstraintsPerColor.Num(); Color++)
+					{
+						PhysicsParallelFor(ConstraintsPerColor[Color].Num(), [&](const int32 Index)
+							{
+								const int32 ConstraintIndex = ConstraintsPerColor[Color][Index];
+								ApplySingleConstraintWithoutSelfTarget(Particles, Dt, ConstraintIndex);
+							});
+					}
+				}
+				else 
+				{
+					for (int32 Color = 0; Color < ConstraintsPerColor.Num(); Color++)
+					{
+						PhysicsParallelFor(ConstraintsPerColor[Color].Num(), [&](const int32 Index)
+							{
+								const int32 ConstraintIndex = ConstraintsPerColor[Color][Index];
+								ApplySingleConstraintWithSelfTarget(Particles, Dt, ConstraintIndex);
+							});
+					}
+				}
+			}
+		}
+
+
+		const TArray<TArray<int32>>& GetIndices()
+		{
+			return Indices;
+		}
+
+		virtual void Init() const
+		{
+			for (T& Lambdas : LambdaArray) { Lambdas = (T)0.; }
+		}
+
+		void UpdateTargets(TArray<TVector<T, 3>>&& InTargets)
+		{
+			Constraints = MoveTemp(InTargets);
+		}
+
+
+	private:
+
+		void InitColor(const ParticleType& Particles)
+		{
+			Chaos::ComputeWeakConstraintsColoring(Indices, SecondIndices, Particles, ConstraintsPerColor);
+		}
+
+		void ApplySingleConstraintWithoutSelfTarget(ParticleType& Particles, const T Dt, const int32 ConstraintIndex) const
+		{
+			T AlphaTilde = T(2) / (Stiffness[ConstraintIndex] * Dt * Dt);
+
+			if (Stiffness[ConstraintIndex] > 1e14)
+			{
+				AlphaTilde = T(0);
+			}
+
+			TVec3<T> SpringEdge((T)0.);
+			for (int32 j = 0; j < Indices[ConstraintIndex].Num(); j++) {
+				for (int32 alpha = 0; alpha < 3; alpha++) {
+					SpringEdge[alpha] += Weights[ConstraintIndex][j] * Particles.P(Indices[ConstraintIndex][j])[alpha];
+				}
+			}
+
+			SpringEdge -= Constraints[ConstraintIndex];
+
+			T Denom = AlphaTilde;
+			for (int32 j = 0; j < Indices[ConstraintIndex].Num(); j++) {
+				Denom += Weights[ConstraintIndex][j] * Weights[ConstraintIndex][j] * Particles.InvM(Indices[ConstraintIndex][j]);
+				
+			}
+
+			for (int32 Beta = 0; Beta < 3; Beta++)
+			{
+				T Cj = SpringEdge[Beta];
+				T DLambda = -Cj - AlphaTilde * LambdaArray[ConstraintIndex * 3 + Beta];
+
+				DLambda /= Denom;
+				LambdaArray[ConstraintIndex * 3 + Beta] += DLambda;
+				for (int32 j = 0; j < Indices[ConstraintIndex].Num(); j++) {
+					Particles.P(Indices[ConstraintIndex][j])[Beta] += DLambda * Weights[ConstraintIndex][j] * Particles.InvM(Indices[ConstraintIndex][j]);
+				}
+			}
+		}
+
+		void ApplySingleConstraintWithSelfTarget(ParticleType& Particles, const T Dt, const int32 ConstraintIndex) const
+		{
+			ensure(SecondIndices.Num() > 0);
+
+			T AlphaTilde = T(2) / (Stiffness[ConstraintIndex] * Dt * Dt);
+
+			if (Stiffness[ConstraintIndex] > 1e14)
+			{
+				AlphaTilde = T(0);
+			}
+
+			TVec3<T> SpringEdge((T)0.);
+			for (int32 j = 0; j < Indices[ConstraintIndex].Num(); j++) {
+				for (int32 alpha = 0; alpha < 3; alpha++) {
+					SpringEdge[alpha] += Weights[ConstraintIndex][j] * Particles.P(Indices[ConstraintIndex][j])[alpha];
+				}
+			}
+
+			for (int32 j = 0; j < SecondIndices[ConstraintIndex].Num(); j++) {
+
+					SpringEdge -= SecondWeights[ConstraintIndex][j] * Particles.P(SecondIndices[ConstraintIndex][j]);
+
+			}
+
+			T Denom = AlphaTilde;
+			for (int32 j = 0; j < Indices[ConstraintIndex].Num(); j++) {
+				Denom += Weights[ConstraintIndex][j] * Weights[ConstraintIndex][j] * Particles.InvM(Indices[ConstraintIndex][j]);
+			}
+
+			for (int32 j = 0; j < SecondIndices[ConstraintIndex].Num(); j++) {
+				Denom += SecondWeights[ConstraintIndex][j] * SecondWeights[ConstraintIndex][j] * Particles.InvM(SecondIndices[ConstraintIndex][j]);
+			}
+
+
+			for (int32 Beta = 0; Beta < 3; Beta++)
+			{
+				T Cj = SpringEdge[Beta];
+				T DLambda = -Cj - AlphaTilde * LambdaArray[ConstraintIndex * 3 + Beta];
+
+				DLambda /= Denom;
+				LambdaArray[ConstraintIndex * 3 + Beta] += DLambda;
+				for (int32 j = 0; j < Indices[ConstraintIndex].Num(); j++) {
+					Particles.P(Indices[ConstraintIndex][j])[Beta] += DLambda * Weights[ConstraintIndex][j] * Particles.InvM(Indices[ConstraintIndex][j]);
+				}
+
+				for (int32 j = 0; j < SecondIndices[ConstraintIndex].Num(); j++) {
+					Particles.P(SecondIndices[ConstraintIndex][j])[Beta] -= DLambda * SecondWeights[ConstraintIndex][j] * Particles.InvM(SecondIndices[ConstraintIndex][j]);
+				}
+			}
+		}
+
+
+	protected:
+		TArray<TArray<int32>> Indices;
+		TArray<TArray<T>> Weights;
+		TArray<TVector<T, 3>> Constraints;
+		TArray<TArray<int32>> SecondIndices;
+		TArray<TArray<T>> SecondWeights;
+		TArray<T> Stiffness;
+		TArray<TArray<int32>> ConstraintsPerColor;
+		mutable TArray<T> LambdaArray;
+	};
+
+}  // End namespace Chaos::Softs
+
diff --git a/Engine/Source/Runtime/Experimental/Chaos/Public/GeometryCollection/Facades/CollectionPositionTargetFacade.h b/Engine/Source/Runtime/Experimental/Chaos/Public/GeometryCollection/Facades/CollectionPositionTargetFacade.h
new file mode 100644
--- /dev/null
+++ b/Engine/Source/Runtime/Experimental/Chaos/Public/GeometryCollection/Facades/CollectionPositionTargetFacade.h
@@ -0,0 +1,77 @@
+// Copyright Epic Games, Inc. All Rights Reserved.
+
+#pragma once
+
+#include "CoreMinimal.h"
+#include "GeometryCollection/ManagedArrayAccessor.h"
+#include "GeometryCollection/ManagedArrayCollection.h"
+#include "GeometryCollection/Facades/CollectionSelectionFacade.h"
+
+namespace GeometryCollection::Facades
+{
+	struct CHAOS_API FPositionTargetsData
+	{
+		TArray<int32> TargetIndex;
+		TArray<int32> SourceIndex;
+		TArray<float> TargetWeights;
+		TArray<float> SourceWeights;
+		FString TargetName;
+		FString SourceName;
+		float Stiffness;
+		float Damping;
+
+	};
+
+	/** Kinematic Facade */
+	class CHAOS_API FPositionTargetFacade
+	{
+	public:
+
+		typedef GeometryCollection::Facades::FSelectionFacade::FSelectionKey FBindingKey;
+
+		//
+		// Kinematics
+		//
+		static const FName GroupName;
+		static const FName TargetIndex;
+		static const FName SourceIndex;
+		static const FName Stiffness;
+		static const FName Damping;
+		static const FName SourceName;
+		static const FName TargetName;
+		static const FName TargetWeights;
+		static const FName SourceWeights;
+
+		FPositionTargetFacade(FManagedArrayCollection& InCollection);
+		FPositionTargetFacade(const FManagedArrayCollection& InCollection);
+
+		/** Create the facade attributes. */
+		void DefineSchema();
+
+		/** Is the facade defined constant. */
+		bool IsConst() const { return Collection == nullptr; }
+
+		/** Is the Facade defined on the collection? */
+		bool IsValid() const;
+
+		//
+		//  Skeletal Mesh Bone Bindings
+		//
+		int32 AddPositionTarget(const FPositionTargetsData& InputData);
+		FPositionTargetsData GetPositionTarget(const int32 DataIndex) const;
+		int32 NumPositionTargets() const { return TargetIndexAttribute.Num(); }
+
+	private:
+		const FManagedArrayCollection& ConstCollection;
+		FManagedArrayCollection* Collection = nullptr;
+
+		TManagedArrayAccessor<TArray<int32>> TargetIndexAttribute;
+		TManagedArrayAccessor<TArray<int32>> SourceIndexAttribute;
+		TManagedArrayAccessor<FString> TargetNameAttribute;
+		TManagedArrayAccessor<FString> SourceNameAttribute;
+		TManagedArrayAccessor<float> StiffnessAttribute;
+		TManagedArrayAccessor<float> DampingAttribute;
+		TManagedArrayAccessor<TArray<float>> TargetWeightsAttribute;
+		TManagedArrayAccessor<TArray<float>> SourceWeightsAttribute;
+	};
+}
