ChaosFlesh: Rename GenerateBindings node to GenerateSurfaceBindings, and add missing passthrough metadata. Make the category name "ChaosDeformable" consistent (with no space), and move solver properties under category "Attributes" to "ChaosDeformable".

diff --git a/Engine/Plugins/Experimental/ChaosFlesh/Source/ChaosFleshEngine/Public/ChaosFlesh/ChaosDeformablePhysicsComponent.h b/Engine/Plugins/Experimental/ChaosFlesh/Source/ChaosFleshEngine/Public/ChaosFlesh/ChaosDeformablePhysicsComponent.h
--- a/Engine/Plugins/Experimental/ChaosFlesh/Source/ChaosFleshEngine/Public/ChaosFlesh/ChaosDeformablePhysicsComponent.h
+++ b/Engine/Plugins/Experimental/ChaosFlesh/Source/ChaosFleshEngine/Public/ChaosFlesh/ChaosDeformablePhysicsComponent.h
@@ -49,21 +49,20 @@ public:
 	UDeformableSolverComponent* GetDeformableSolver();
 	const UDeformableSolverComponent* GetDeformableSolver() const;
 	
-	UPROPERTY(EditAnywhere, BlueprintReadWrite, Category = "Chaos Deformable")
+	UPROPERTY(EditAnywhere, BlueprintReadWrite, Category = "ChaosDeformable")
 	TObjectPtr<UDeformableSolverComponent> PrimarySolverComponent;
 
-	UPROPERTY(EditAnywhere, BlueprintReadWrite, Category = "Chaos Deformable")
+	UPROPERTY(EditAnywhere, BlueprintReadWrite, Category = "ChaosDeformable")
 	bool bTempEnableGravity = true;
 
-	UPROPERTY(EditAnywhere, BlueprintReadWrite, Category = Attributes, meta = (ClampMin = "0.0", ClampMax = "1.0", UIMin = "0.0", UIMax = "1.0"))
-		float DampingMultiplier = 1.f;
+	UPROPERTY(EditAnywhere, BlueprintReadWrite, Category = "ChaosDeformable", meta = (ClampMin = "0.0", ClampMax = "1.0", UIMin = "0.0", UIMax = "1.0"))
+	float DampingMultiplier = 1.f;
 
-	UPROPERTY(EditAnywhere, BlueprintReadWrite, Category = Attributes, meta = (ClampMin = "0.0", ClampMax = "1.0", UIMin = "0.0", UIMax = "1.0"))
-		float StiffnessMultiplier = 1.f;
-
-	UPROPERTY(EditAnywhere, BlueprintReadWrite, Category = Attributes, meta = (ClampMin = "0.0", ClampMax = "1.0", UIMin = "0.0", UIMax = "1.0"))
-		float MassMultiplier = 1.f;
+	UPROPERTY(EditAnywhere, BlueprintReadWrite, Category = "ChaosDeformable", meta = (ClampMin = "0.0", ClampMax = "1.0", UIMin = "0.0", UIMax = "1.0"))
+	float StiffnessMultiplier = 1.f;
 
+	UPROPERTY(EditAnywhere, BlueprintReadWrite, Category = "ChaosDeformable", meta = (ClampMin = "0.0", ClampMax = "1.0", UIMin = "0.0", UIMax = "1.0"))
+	float MassMultiplier = 1.f;
 
 	const FThreadingProxy* GetPhysicsProxy() const { return PhysicsProxy; }
 	FThreadingProxy* GetPhysicsProxy() { return PhysicsProxy; }
diff --git a/Engine/Plugins/Experimental/ChaosFlesh/Source/ChaosFleshNodes/Private/Dataflow/ChaosFleshBindingsNodes.cpp b/Engine/Plugins/Experimental/ChaosFlesh/Source/ChaosFleshNodes/Private/Dataflow/ChaosFleshBindingsNodes.cpp
--- a/Engine/Plugins/Experimental/ChaosFlesh/Source/ChaosFleshNodes/Private/Dataflow/ChaosFleshBindingsNodes.cpp
+++ b/Engine/Plugins/Experimental/ChaosFlesh/Source/ChaosFleshNodes/Private/Dataflow/ChaosFleshBindingsNodes.cpp
@@ -23,7 +23,7 @@ namespace Dataflow
 {
 	void ChaosFleshBindingsNodes()
 	{
-		DATAFLOW_NODE_REGISTER_CREATION_FACTORY(FGenerateBindings);
+		DATAFLOW_NODE_REGISTER_CREATION_FACTORY(FGenerateSurfaceBindings);
 	}
 }
 
@@ -107,7 +107,7 @@ BuildVertexToVertexAdjacencyBuffer(
 }
 
 void
-FGenerateBindings::Evaluate(Dataflow::FContext& Context, const FDataflowOutput* Out) const
+FGenerateSurfaceBindings::Evaluate(Dataflow::FContext& Context, const FDataflowOutput* Out) const
 {
 	if (Out->IsA<DataType>(&Collection))
 	{
diff --git a/Engine/Plugins/Experimental/ChaosFlesh/Source/ChaosFleshNodes/Public/Dataflow/ChaosFleshBindingsNodes.h b/Engine/Plugins/Experimental/ChaosFlesh/Source/ChaosFleshNodes/Public/Dataflow/ChaosFleshBindingsNodes.h
--- a/Engine/Plugins/Experimental/ChaosFlesh/Source/ChaosFleshNodes/Public/Dataflow/ChaosFleshBindingsNodes.h
+++ b/Engine/Plugins/Experimental/ChaosFlesh/Source/ChaosFleshNodes/Public/Dataflow/ChaosFleshBindingsNodes.h
@@ -15,17 +15,17 @@ DECLARE_LOG_CATEGORY_EXTERN(LogMeshBindings, Verbose, All);
 
 // Generate barycentric bindings (used by the FleshDeformer deformer graph) of a render surface to a tetrahedral mesh.
 USTRUCT(meta = (DataflowFlesh))
-struct FGenerateBindings : public FDataflowNode
+struct FGenerateSurfaceBindings : public FDataflowNode
 {
 	GENERATED_USTRUCT_BODY()
-	DATAFLOW_NODE_DEFINE_INTERNAL(FGenerateBindings, "GenerateBindings", "Flesh", "")
+	DATAFLOW_NODE_DEFINE_INTERNAL(FGenerateSurfaceBindings, "GenerateSurfaceBindings", "Flesh", "")
 	DATAFLOW_NODE_RENDER_TYPE(FGeometryCollection::StaticType(), "Collection")
 
 public:
 	typedef FManagedArrayCollection DataType;
 
 	// Passthrough geometry collection. Bindings are stored as standalone groups in the \p Collection, keyed by the name of the input render mesh and all available LOD's.
-	UPROPERTY(meta = (DataflowInput, DataflowOutput, DisplayName = "Collection"))
+	UPROPERTY(meta = (DataflowInput, DataflowOutput, DataflowPassthrough = "Collection", DisplayName = "Collection"))
 	FManagedArrayCollection Collection;
 
 	// The input mesh, whose render surface is used to generate bindings.
@@ -41,18 +41,18 @@ public:
 	bool bDoSurfaceProjection = true;
 
 	// The maximum number of iterations to try expanding the domain while looking for surface triangles to bind to.
-	UPROPERTY(EditAnywhere, Category = "Dataflow", meta = (DisplayName = "SurfaceProjectionIterations"))
+	UPROPERTY(EditAnywhere, Category = "Dataflow", meta = (DisplayName = "SurfaceProjectionIterations", EditCondition = "bDoSurfaceProjection == true"))
 	uint32 SurfaceProjectionIterations = 10;
 
 	// When nodes aren't contained in tetrahedra and surface projection fails, try to find suitable bindings by looking to neighboring parents.
 	UPROPERTY(EditAnywhere, Category = "Dataflow", meta = (DisplayName = "DoOrphanReparenting"))
 	bool bDoOrphanReparenting = true;
 
-	FGenerateBindings(const Dataflow::FNodeParameters& InParam, FGuid InGuid = FGuid::NewGuid())
+	FGenerateSurfaceBindings(const Dataflow::FNodeParameters& InParam, FGuid InGuid = FGuid::NewGuid())
 		: FDataflowNode(InParam, InGuid)
 	{
 		RegisterInputConnection(&Collection);
-		RegisterOutputConnection(&Collection);
+		RegisterOutputConnection(&Collection, &Collection);
 		RegisterInputConnection(&StaticMeshIn);
 		RegisterInputConnection(&SkeletalMeshIn);
 	}
diff --git a/Engine/Source/Runtime/Experimental/Dataflow/Core/Private/Dataflow/DataflowNode.cpp b/Engine/Source/Runtime/Experimental/Dataflow/Core/Private/Dataflow/DataflowNode.cpp
--- a/Engine/Source/Runtime/Experimental/Dataflow/Core/Private/Dataflow/DataflowNode.cpp
+++ b/Engine/Source/Runtime/Experimental/Dataflow/Core/Private/Dataflow/DataflowNode.cpp
@@ -351,7 +351,7 @@ bool FDataflowNode::ValidateConnections()
 					}
 					else if(OutputConnection->GetPassthroughRealAddress()) 
 					{
-						UE_LOG(LogChaos, Warning, TEXT("Missing DataflowPassthrough decleration for (%s:%s)"), *GetName().ToString(),*PropName.ToString());
+						UE_LOG(LogChaos, Warning, TEXT("Missing DataflowPassthrough declaration for (%s:%s)"), *GetName().ToString(),*PropName.ToString());
 						bValid = false;
 					}
 				}
