Project Update : SkeletalBindings
- Added node to hold barycentric weights from a embedded skeletal mesh.
- Dataflow force invalidate a node when Evaluate is called on a toolkit node.
- Removed older shared file.

diff --git a/Engine/Plugins/Experimental/ChaosFlesh/Shaders/Private/DataInterfaceFlesh.ush b/Engine/Plugins/Experimental/ChaosFlesh/Shaders/Private/DataInterfaceFlesh.ush
deleted file mode 100644
diff --git a/Engine/Plugins/Experimental/ChaosFlesh/Source/ChaosFlesh/Private/ChaosFlesh/TetrahedralCollection.cpp b/Engine/Plugins/Experimental/ChaosFlesh/Source/ChaosFlesh/Private/ChaosFlesh/TetrahedralCollection.cpp
--- a/Engine/Plugins/Experimental/ChaosFlesh/Source/ChaosFlesh/Private/ChaosFlesh/TetrahedralCollection.cpp
+++ b/Engine/Plugins/Experimental/ChaosFlesh/Source/ChaosFlesh/Private/ChaosFlesh/TetrahedralCollection.cpp
@@ -104,6 +104,28 @@ void FTetrahedralCollection::Init(
 	}
 }
 
+void FTetrahedralCollection::UpdateBoundingBox()
+{
+	if (BoundingBox.Num())
+	{
+		// Initialize BoundingBox
+		for (int32 Idx = 0; Idx < BoundingBox.Num(); ++Idx)
+		{
+			BoundingBox[Idx].Init();
+		}
+
+		// Compute BoundingBox
+		for (int32 Idx = 0; Idx < Vertex.Num(); ++Idx)
+		{
+			int32 TransformIndexValue = BoneMap[Idx];
+			if (TransformToGeometryIndex[TransformIndexValue] != INDEX_NONE)
+			{
+				BoundingBox[TransformToGeometryIndex[TransformIndexValue]] += FVector(Vertex[Idx]);
+			}
+		}
+	}
+}
+
 int32 FTetrahedralCollection::AppendGeometry(
 	const FTetrahedralCollection& Other, 
 	int32 MaterialIDOffset, 
diff --git a/Engine/Plugins/Experimental/ChaosFlesh/Source/ChaosFlesh/Public/ChaosFlesh/TetrahedralCollection.h b/Engine/Plugins/Experimental/ChaosFlesh/Source/ChaosFlesh/Public/ChaosFlesh/TetrahedralCollection.h
--- a/Engine/Plugins/Experimental/ChaosFlesh/Source/ChaosFlesh/Public/ChaosFlesh/TetrahedralCollection.h
+++ b/Engine/Plugins/Experimental/ChaosFlesh/Source/ChaosFlesh/Public/ChaosFlesh/TetrahedralCollection.h
@@ -33,6 +33,8 @@ public:
 	static FTetrahedralCollection* NewTetrahedralCollection(const TArray<FVector>& Vertices, const TArray<FIntVector3>& SurfaceElements, const TArray<FIntVector4>& Elements,bool bReverseVertexOrder = true);
 	static void Init(FTetrahedralCollection* Collection, const TArray<FVector>& Vertices, const TArray<FIntVector3>& SurfaceElements, const TArray<FIntVector4>& Elements, bool bReverseVertexOrder = true);
 
+	void UpdateBoundingBox();
+
 	int32 AppendGeometry(const FTetrahedralCollection& GeometryCollection, int32 MaterialIDOffset = 0, bool ReindexAllMaterials = true, const FTransform& TransformRoot = FTransform::Identity);
 
 	/**
diff --git a/Engine/Plugins/Experimental/ChaosFlesh/Source/ChaosFleshEngine/Private/ChaosFlesh/FleshComponent.cpp b/Engine/Plugins/Experimental/ChaosFlesh/Source/ChaosFleshEngine/Private/ChaosFlesh/FleshComponent.cpp
--- a/Engine/Plugins/Experimental/ChaosFlesh/Source/ChaosFleshEngine/Private/ChaosFlesh/FleshComponent.cpp
+++ b/Engine/Plugins/Experimental/ChaosFlesh/Source/ChaosFleshEngine/Private/ChaosFlesh/FleshComponent.cpp
@@ -4,24 +4,28 @@
 
 #include "Animation/SkeletalMeshActor.h"
 #include "Animation/Skeleton.h"
+#include "ChaosStats.h"
 #include "Chaos/DebugDrawQueue.h"
+#include "Chaos/Tetrahedron.h"
 #include "ChaosFlesh/FleshCollection.h"
 #include "ChaosFlesh/ChaosDeformableSolverActor.h"
 #include "ChaosFlesh/ChaosDeformableSolverComponent.h"
-#include "ChaosStats.h"
+#include "ChaosFlesh/ChaosDeformableTypes.h"
+#include "ChaosFlesh/FleshDynamicAsset.h"
 #include "Components/SkeletalMeshComponent.h"
 #include "Dataflow/DataflowEngineUtil.h"
 #include "Engine/SkeletalMesh.h"
 #include "GeometryCollection/Facades/CollectionTransformSourceFacade.h"
+#include "GeometryCollection/Facades/CollectionTetrahedralSkeletalBindingsFacade.h"
 #include "GeometryCollection/TransformCollection.h"
 #include "ProceduralMeshComponent.h"
-#include "ChaosFlesh/ChaosDeformableTypes.h"
 
 
 #include UE_INLINE_GENERATED_CPP_BY_NAME(FleshComponent)
 
 FChaosEngineDeformableCVarParams CVarParams;
 FAutoConsoleVariableRef CVarDeforambleDoDrawSimulationMesh(TEXT("p.Chaos.DebugDraw.Deformable.SimulationMesh"), CVarParams.bDoDrawSimulationMesh, TEXT("Debug draw the deformable simulation resutls on the game thread. [def: true]"));
+FAutoConsoleVariableRef CVarDeforambleDoDrawSkeletalMeshBindingPositions(TEXT("p.Chaos.DebugDraw.Deformable.SkeletalMeshBindingPositions"), CVarParams.bDoDrawSkeletalMeshBindingPositions, TEXT("Debug draw the deformable simulation's SkeletalMeshBindingPositions on the game thread. [def: false]"));
 
 DECLARE_CYCLE_STAT(TEXT("Chaos.Deformable.UFleshComponent.TickComponent"), STAT_ChaosDeformable_UFleshComponent_TickComponent, STATGROUP_Chaos);
 DECLARE_CYCLE_STAT(TEXT("Chaos.Deformable.UFleshComponent.NewProxy"), STAT_ChaosDeformable_UFleshComponent_NewProxy, STATGROUP_Chaos);
@@ -218,12 +222,16 @@ void UFleshComponent::UpdateFromSimualtion(const FDataMapValue* SimualtionBuffer
 			for (int i = DynamicVertex.Num() - 1; i >= 0; i--)
 			{
 				DynamicVertex[i] = UEVertf(GetComponentTransform().InverseTransformPosition(UEVertd(SimulationVertex[i])));
-				//DynamicVertex[i] = UEVertf(PrevTransform.InverseTransformPosition(UEVertd(SimulationVertex[i])));
-				//Chaos::FDebugDrawQueue::GetInstance().DrawDebugPoint(UEVertd(SimulationVertex[i]), FColor::Red, false, -1.0f, 0, 5);
+			}
+
+			 //p.Chaos.DebugDraw.Enabled 1
+			 //p.Chaos.DebugDraw.Deformable.SkeletalMeshBindingPositions 1
+			if (CVarParams.bDoDrawSkeletalMeshBindingPositions)
+			{
+				DebugDrawSkeletalMeshBindingPositions();
 			}
 		}
 	}
-	//PrevTransform = this->GetComponentTransform();
 }
 
 void UFleshComponent::UpdateLocalBounds()
@@ -401,6 +409,79 @@ void UFleshComponent::RenderProceduralMesh()
 }
 
 
+void UFleshComponent::DebugDrawSkeletalMeshBindingPositions() const
+{
+#if WITH_EDITOR
+	auto UEVertd = [](FVector3f V) { return FVector3d(V.X, V.Y, V.Z); };
+
+	if (const UFleshAsset* RestAsset = GetRestCollection())
+	{
+		const USkeletalMesh* SkeletalMesh = TargetDeformationSkeleton ? TargetDeformationSkeleton : RestAsset->SkeletalMesh;
+		if (SkeletalMesh)
+		{
+			TArray<bool> Influenced;
+			TArray<FVector> PosArray = GetSkeletalMeshBindingPositionsInternal(SkeletalMesh, &Influenced);
+			for (int i=0;i<PosArray.Num();i++)
+			{
+				const FVector& Pos = PosArray[i];
+				if (Influenced[i])
+				{
+					Chaos::FDebugDrawQueue::GetInstance().DrawDebugPoint(GetComponentTransform().TransformPosition(Pos), FColor::Red, true, 2.0f, SDPG_Foreground, 10);
+				}
+			}
+		}
+	}
+#endif
+}
+
+TArray<FVector> UFleshComponent::GetSkeletalMeshBindingPositions(const USkeletalMesh* InSkeletalMesh) const
+{
+	return GetSkeletalMeshBindingPositionsInternal(InSkeletalMesh);
+}
+
+TArray<FVector> UFleshComponent::GetSkeletalMeshBindingPositionsInternal(const USkeletalMesh* InSkeletalMesh, TArray<bool>* OutInfluence) const
+{
+	auto UEVert3d = [](FVector3f V) { return FVector3d(V.X, V.Y, V.Z); };
+	auto UEVert4d = [](FVector4f V) { return FVector4d(V.X, V.Y, V.Z, V.W); };
+
+	TArray<FVector> TransformPositions;
+	if (InSkeletalMesh)
+	{
+		FName SkeletalMeshName(InSkeletalMesh->GetName());
+		if (const UFleshAsset* RestAsset = GetRestCollection())
+		{
+			if (const FFleshCollection* Rest = RestAsset->GetCollection())
+			{
+				GeometryCollection::Facades::FTetrahedralSkeletalBindings TetBindings(*Rest);
+
+				const TManagedArray<int32>* TetrahedronStart = Rest->FindAttribute<int32>(FTetrahedralCollection::TetrahedronStartAttribute, FGeometryCollection::GeometryGroup);
+				const TManagedArray<FVector3f>* Verts = GetDynamicCollection() ? GetDynamicCollection()->FindPositions() : Rest->FindAttributeTyped<FVector3f>("Vertex", FGeometryCollection::VerticesGroup);
+
+				if (ensure(Verts != nullptr) && TetrahedronStart)
+				{
+					TArray<FTransform> ComponentPose;
+					Dataflow::Animation::GlobalTransforms(InSkeletalMesh->GetRefSkeleton(), ComponentPose);
+
+					TransformPositions.SetNumUninitialized(ComponentPose.Num());
+					for (int32 i = 0; i < ComponentPose.Num(); i++)
+					{
+						TransformPositions[i] = ComponentPose[i].GetTranslation();
+					}
+
+					if (OutInfluence != nullptr) OutInfluence->Init(false, TransformPositions.Num());
+					for (int32 TetMeshIdx = 0; TetMeshIdx < TetrahedronStart->Num(); TetMeshIdx++)
+					{
+						FString MeshBindingsName = GeometryCollection::Facades::FTetrahedralSkeletalBindings::GenerateMeshGroupName(TetMeshIdx, SkeletalMeshName);
+						TetBindings.CalculateBindings(MeshBindingsName, Verts->GetConstArray(), TransformPositions, OutInfluence);
+					}
+				}
+			}
+		}
+	}
+	return TransformPositions;
+}
+
+
 
 
 
diff --git a/Engine/Plugins/Experimental/ChaosFlesh/Source/ChaosFleshEngine/Private/ChaosFlesh/FleshDynamicAsset.cpp b/Engine/Plugins/Experimental/ChaosFlesh/Source/ChaosFleshEngine/Private/ChaosFlesh/FleshDynamicAsset.cpp
--- a/Engine/Plugins/Experimental/ChaosFlesh/Source/ChaosFleshEngine/Private/ChaosFlesh/FleshDynamicAsset.cpp
+++ b/Engine/Plugins/Experimental/ChaosFlesh/Source/ChaosFleshEngine/Private/ChaosFlesh/FleshDynamicAsset.cpp
@@ -50,6 +50,11 @@ TManagedArray<FVector3f>& UFleshDynamicAsset::GetPositions()
 	return DynamicCollection->ModifyAttribute<FVector3f>("Vertex", FGeometryCollection::VerticesGroup);
 }
 
+const TManagedArray<FVector3f>* UFleshDynamicAsset::FindPositions() const
+{
+	return DynamicCollection->FindAttributeTyped<FVector3f>("Vertex", FGeometryCollection::VerticesGroup);
+}
+
 TManagedArray<int32>& UFleshDynamicAsset::GetObjectState()
 {
 	return DynamicCollection->ModifyAttribute<int32>("ObjectState", FGeometryCollection::VerticesGroup);
diff --git a/Engine/Plugins/Experimental/ChaosFlesh/Source/ChaosFleshEngine/Public/ChaosFlesh/ChaosDeformableTypes.h b/Engine/Plugins/Experimental/ChaosFlesh/Source/ChaosFleshEngine/Public/ChaosFlesh/ChaosDeformableTypes.h
--- a/Engine/Plugins/Experimental/ChaosFlesh/Source/ChaosFleshEngine/Public/ChaosFlesh/ChaosDeformableTypes.h
+++ b/Engine/Plugins/Experimental/ChaosFlesh/Source/ChaosFleshEngine/Public/ChaosFlesh/ChaosDeformableTypes.h
@@ -25,4 +25,5 @@ struct FChaosEngineDeformableCVarParams
 #else
 	bool bDoDrawSimulationMesh = false;
 #endif
+	bool bDoDrawSkeletalMeshBindingPositions = false;
 };
\ No newline at end of file
diff --git a/Engine/Plugins/Experimental/ChaosFlesh/Source/ChaosFleshEngine/Public/ChaosFlesh/FleshComponent.h b/Engine/Plugins/Experimental/ChaosFlesh/Source/ChaosFleshEngine/Public/ChaosFlesh/FleshComponent.h
--- a/Engine/Plugins/Experimental/ChaosFlesh/Source/ChaosFleshEngine/Public/ChaosFlesh/FleshComponent.h
+++ b/Engine/Plugins/Experimental/ChaosFlesh/Source/ChaosFleshEngine/Public/ChaosFlesh/FleshComponent.h
@@ -29,6 +29,10 @@ public:
 
 	~UFleshComponent();
 
+	UFUNCTION(BlueprintCallable, Category = "Physics")
+	TArray<FVector> GetSkeletalMeshBindingPositions(const USkeletalMesh* InSkeletalMesh) const;
+
+	/** USceneComponent Interface */
 	virtual void BeginPlay() override;
 	virtual void OnRegister() override;
 	virtual void TickComponent(float DeltaTime, enum ELevelTick TickType, FActorComponentTickFunction *ThisTickFunction) override;
@@ -57,6 +61,10 @@ public:
 	UPROPERTY(EditAnywhere, BlueprintReadWrite, Category = Attributes)
 	TObjectPtr<UProceduralMeshComponent> Mesh;
 
+	UPROPERTY(EditAnywhere, BlueprintReadWrite, Category = "Rendering")
+	TObjectPtr<USkeletalMesh> TargetDeformationSkeleton;
+
+
 private:
 	UPROPERTY(EditAnywhere, Category = "ChaosPhysics")
 	TObjectPtr<const UFleshAsset> RestCollection;
@@ -84,5 +92,8 @@ private:
 	FBoxSphereBounds BoundingBox = FBoxSphereBounds(ForceInitToZero);
 
 	FTransform PrevTransform = FTransform::Identity;
+
+	TArray<FVector> GetSkeletalMeshBindingPositionsInternal(const USkeletalMesh* InSkeletalMesh, TArray<bool>* OutInfluence = nullptr) const;
+	void DebugDrawSkeletalMeshBindingPositions() const;
 };
 
diff --git a/Engine/Plugins/Experimental/ChaosFlesh/Source/ChaosFleshEngine/Public/ChaosFlesh/FleshDynamicAsset.h b/Engine/Plugins/Experimental/ChaosFlesh/Source/ChaosFleshEngine/Public/ChaosFlesh/FleshDynamicAsset.h
--- a/Engine/Plugins/Experimental/ChaosFlesh/Source/ChaosFleshEngine/Public/ChaosFlesh/FleshDynamicAsset.h
+++ b/Engine/Plugins/Experimental/ChaosFlesh/Source/ChaosFleshEngine/Public/ChaosFlesh/FleshDynamicAsset.h
@@ -31,6 +31,7 @@ public:
 	void ResetAttributesFrom(const FManagedArrayCollection* InCopyFrom);
 
 	TManagedArray<FVector3f>& GetPositions();
+	const TManagedArray<FVector3f>* FindPositions() const;
 	TManagedArray<int32>& GetObjectState();
 
 	const FManagedArrayCollection* GetCollection() const { return DynamicCollection.Get(); }
diff --git a/Engine/Plugins/Experimental/ChaosFlesh/Source/ChaosFleshNodes/Private/Dataflow/ChaosFleshNodesPlugin.cpp b/Engine/Plugins/Experimental/ChaosFlesh/Source/ChaosFleshNodes/Private/Dataflow/ChaosFleshNodesPlugin.cpp
--- a/Engine/Plugins/Experimental/ChaosFlesh/Source/ChaosFleshNodes/Private/Dataflow/ChaosFleshNodesPlugin.cpp
+++ b/Engine/Plugins/Experimental/ChaosFlesh/Source/ChaosFleshNodes/Private/Dataflow/ChaosFleshNodesPlugin.cpp
@@ -11,6 +11,7 @@
 #include "Dataflow/ChaosFleshRadialTetrahedronNodes.h"
 #include "Dataflow/ChaosFleshRenderInitializationNodes.h"
 #include "Dataflow/ChaosFleshPositionTargetInitializationNodes.h"
+#include "Dataflow/ChaosFleshSkeletalBindingsNode.h"
 #include "Dataflow/ChaosFleshTetrahedralNodes.h"
 #include "Modules/ModuleManager.h"
 
@@ -19,6 +20,7 @@
 
 void IChaosFleshNodesPlugin::StartupModule()
 {
+	Dataflow::ChaosFleshBindingsNodes();
 	Dataflow::RegisterChaosFleshEngineAssetNodes();
 	Dataflow::RegisterChaosFleshCoreNodes();
 	Dataflow::ChaosFleshFiberDirectionInitializationNodes();
@@ -26,7 +28,7 @@ void IChaosFleshNodesPlugin::StartupModule()
 	Dataflow::RegisterChaosFleshKinematicInitializationNodes();
 	Dataflow::RegisterChaosFleshPositionTargetInitializationNodes();
 	Dataflow::ChaosFleshTetrahedralNodes();
-	Dataflow::ChaosFleshBindingsNodes();
+	Dataflow::ChaosFleshSkeletalBindingsNode();
 	Dataflow::ChaosFleshRadialTetrahedronNodes();
 }
 
diff --git a/Engine/Plugins/Experimental/ChaosFlesh/Source/ChaosFleshNodes/Private/Dataflow/ChaosFleshSkeletalBindingsNode.cpp b/Engine/Plugins/Experimental/ChaosFlesh/Source/ChaosFleshNodes/Private/Dataflow/ChaosFleshSkeletalBindingsNode.cpp
new file mode 100644
--- /dev/null
+++ b/Engine/Plugins/Experimental/ChaosFlesh/Source/ChaosFleshNodes/Private/Dataflow/ChaosFleshSkeletalBindingsNode.cpp
@@ -0,0 +1,184 @@
+// Copyright Epic Games, Inc. All Rights Reserved.
+
+#include "Dataflow/ChaosFleshSkeletalBindingsNode.h"
+
+#include "Chaos/AABBTree.h"
+#include "Chaos/BoundingVolumeHierarchy.h"
+#include "Chaos/Tetrahedron.h"
+#include "Chaos/TriangleMesh.h"
+#include "ChaosFlesh/TetrahedralCollection.h"
+#include "Containers/Map.h"
+#include "Engine/StaticMesh.h"
+#include "Dataflow/DataflowInputOutput.h"
+#include "Engine/SkeletalMesh.h"
+#include "Dataflow/DataflowNodeFactory.h"
+#include "GeometryCollection/Facades/CollectionTetrahedralSkeletalBindingsFacade.h"
+#include "Rendering/SkeletalMeshRenderData.h"
+#include "UObject/PrimaryAssetId.h"
+
+DEFINE_LOG_CATEGORY(LogSkeletalBindings);
+
+
+namespace Dataflow
+{
+	void ChaosFleshSkeletalBindingsNode()
+	{
+		DATAFLOW_NODE_REGISTER_CREATION_FACTORY(FGenerateSkeletalBindings);
+	}
+
+
+
+	TArray<int32> ChildIndices(TObjectPtr<const USkeletalMesh> SkeletalMesh, int32 StartIndex)
+	{
+		int32 IndexCount = 0;
+		TArray<int32> BranchIndices;
+
+		TArray<int32> ToProcess;
+		int32 CurrentIndex = StartIndex;
+		while (SkeletalMesh->GetRefSkeleton().IsValidIndex(CurrentIndex))
+		{
+			TArray<int32> Buffer;
+			SkeletalMesh->GetRefSkeleton().GetDirectChildBones(CurrentIndex, Buffer);
+
+			if (Buffer.Num())
+			{
+				ToProcess.Append(Buffer);
+			}
+
+			if (!IndexCount) // skip the first index to get only the chidren. 
+			{
+				BranchIndices.Add(CurrentIndex);
+			}
+
+			CurrentIndex = INDEX_NONE;
+			if (ToProcess.Num())
+			{
+				CurrentIndex = ToProcess.Pop();
+			}
+		}
+		return BranchIndices;
+	}
+
+}
+
+void
+FGenerateSkeletalBindings::Evaluate(Dataflow::FContext& Context, const FDataflowOutput* Out) const
+{
+	auto UEVertd = [](FVector3f V) { return FVector3d(V.X, V.Y, V.Z); };
+	auto UEVertf = [](FVector3d V) { return FVector3f(V.X, V.Y, V.Z); };
+
+	if (Out->IsA<DataType>(&Collection))
+	{
+		DataType OutCollection = GetValue<DataType>(Context, &Collection); // Deep copy
+
+		TManagedArray<FIntVector4>* Tetrahedron = OutCollection.FindAttribute<FIntVector4>(FTetrahedralCollection::TetrahedronAttribute, FTetrahedralCollection::TetrahedralGroup);
+		TManagedArray<int32>* TetrahedronStart =OutCollection.FindAttribute<int32>(FTetrahedralCollection::TetrahedronStartAttribute, FGeometryCollection::GeometryGroup);
+		TManagedArray<int32>* TetrahedronCount =OutCollection.FindAttribute<int32>(FTetrahedralCollection::TetrahedronCountAttribute, FGeometryCollection::GeometryGroup);
+		TManagedArray<TArray<int32>>* IncidentElements = OutCollection.FindAttribute<TArray<int32>>(FTetrahedralCollection::IncidentElementsAttribute, FGeometryCollection::VerticesGroup);
+		TManagedArray<FVector3f>* Vertex = OutCollection.FindAttribute<FVector3f>("Vertex", "Vertices");
+
+		TObjectPtr<const USkeletalMesh> SkeletalMesh = GetValue<TObjectPtr<const USkeletalMesh>>(Context, &SkeletalMeshIn);
+		if (SkeletalMesh && Tetrahedron && TetrahedronStart && TetrahedronCount && Vertex)
+		{
+			// Iterate all the tetrahedral meshs
+			for (int32 TetMeshIdx = 0; TetMeshIdx < TetrahedronStart->Num(); TetMeshIdx++)
+			{
+				const int32 TetMeshStart = (*TetrahedronStart)[TetMeshIdx];
+				const int32 TetMeshCount = (*TetrahedronCount)[TetMeshIdx];
+
+				// Build Tetrahedra
+				TArray<Chaos::TTetrahedron<Chaos::FReal>> Tets;			// Index 0 == TetMeshStart
+				TArray<Chaos::TTetrahedron<Chaos::FReal>*> BVHTetPtrs;
+				Tets.SetNumUninitialized(TetMeshCount);
+				BVHTetPtrs.SetNumUninitialized(TetMeshCount);
+				for (int32 i = 0; i < TetMeshCount; i++)
+				{
+					const int32 Idx = TetMeshStart + i;
+					const FIntVector4& Tet = (*Tetrahedron)[Idx];
+					Tets[i] = Chaos::TTetrahedron<Chaos::FReal>(
+						(*Vertex)[Tet[0]],
+						(*Vertex)[Tet[1]],
+						(*Vertex)[Tet[2]],
+						(*Vertex)[Tet[3]]);
+					BVHTetPtrs[i] = &Tets[i];
+				}
+
+				// Init BVH for tetrahedra.
+				Chaos::TBoundingVolumeHierarchy<
+					TArray<Chaos::TTetrahedron<Chaos::FReal>*>,
+					TArray<int32>,
+					Chaos::FReal,
+					3> TetBVH(BVHTetPtrs);
+
+				FName SkeletalMeshName = FName(SkeletalMesh->GetName());
+				TArray<FTransform> ComponentPose;
+				Dataflow::Animation::GlobalTransforms(SkeletalMesh->GetRefSkeleton(), ComponentPose);
+				TArray<int32> Indices = Dataflow::ChildIndices(SkeletalMesh, GetValue<int32>(Context, &BoneIndexIn));
+
+				//
+				// Do intersection tests against tets
+				//
+
+				int CurrentIndex = 0;
+				TArray<int32> TetIndex; TetIndex.SetNum(Indices.Num());
+				TArray<FVector4f> Weights;		Weights.SetNum(Indices.Num());
+				TArray<FVector3f> Offsets;		Offsets.SetNum(Indices.Num());
+				TArray<int32> IndexKey;			IndexKey.SetNum(Indices.Num());
+
+				TArray<int32> TetIntersections; TetIntersections.Reserve(64);
+				for (int32 i = 0; i < Indices.Num(); i++)
+				{
+					int32 TransformIndex = Indices[i];
+
+					FVector TestPosition = ComponentPose[TransformIndex].GetTranslation();
+					TetIntersections = TetBVH.FindAllIntersections(TestPosition);
+					for (int32 j = 0; j < TetIntersections.Num(); j++)
+					{
+						const int32 TetIdx = TetIntersections[j];
+						if (!Tets[TetIdx].Outside(TestPosition, 1.0e-2)) // includes boundary
+						{
+							Chaos::TVector<Chaos::FReal, 4> WeightsD = Tets[TetIdx].GetBarycentricCoordinates(TestPosition);
+							Weights[CurrentIndex] = FVector4f(WeightsD[0], WeightsD[1], WeightsD[2], WeightsD[3]);
+							TetIndex[CurrentIndex] = TetIdx + TetMeshStart;
+							IndexKey[CurrentIndex] = TransformIndex;
+
+							// validate
+							{
+								FIntVector4 Parent = (*Tetrahedron)[TetIndex[CurrentIndex]];
+								FVector3f EmbeddedPos =
+									(*Vertex)[Parent[0]] * Weights[CurrentIndex][0] +
+									(*Vertex)[Parent[1]] * Weights[CurrentIndex][1] +
+									(*Vertex)[Parent[2]] * Weights[CurrentIndex][2] +
+									(*Vertex)[Parent[3]] * Weights[CurrentIndex][3];
+								ensureMsgf((UEVertf(TestPosition) - EmbeddedPos).SquaredLength() < 1.0, TEXT("ERROR: Skeletal Bindings mis alignment."));
+								check((UEVertf(TestPosition) - EmbeddedPos).SquaredLength() < 1.0);
+							}
+
+							CurrentIndex++;
+							break;
+						}
+					}
+				}
+				TetIndex.SetNum(CurrentIndex, false);
+				Weights.SetNum(CurrentIndex, false);
+				IndexKey.SetNum(CurrentIndex, false);
+
+				UE_LOG(LogSkeletalBindings, Display,
+					   TEXT("'%s' - Generated mesh bindings between tet mesh index %d and skeletal mesh - stats:\n"
+					   "    Skeletal transforms num: %d\n"
+					   "    Transforms in tetrahedra: %d"),
+					   *(SkeletalMeshName.ToString()), TetMeshIdx, Indices.Num(), CurrentIndex);
+
+				// Store bindings in the collection
+				if (TetIndex.Num())
+				{
+					using namespace GeometryCollection::Facades;
+					GeometryCollection::Facades::FTetrahedralSkeletalBindings TetBindings(OutCollection);
+					FString MeshBindingsName = FTetrahedralSkeletalBindings::GenerateMeshGroupName(TetMeshIdx, SkeletalMeshName);
+					TetBindings.SetBindings(MeshBindingsName, TetIndex, Weights, IndexKey);
+				}
+			} // end for TetMeshIdx
+		}
+		SetValue<DataType>(Context, OutCollection, &Collection);
+	}
+}
diff --git a/Engine/Plugins/Experimental/ChaosFlesh/Source/ChaosFleshNodes/Public/Dataflow/ChaosFleshPositionTargetInitializationNodes.h b/Engine/Plugins/Experimental/ChaosFlesh/Source/ChaosFleshNodes/Public/Dataflow/ChaosFleshPositionTargetInitializationNodes.h
--- a/Engine/Plugins/Experimental/ChaosFlesh/Source/ChaosFleshNodes/Public/Dataflow/ChaosFleshPositionTargetInitializationNodes.h
+++ b/Engine/Plugins/Experimental/ChaosFlesh/Source/ChaosFleshNodes/Public/Dataflow/ChaosFleshPositionTargetInitializationNodes.h
@@ -40,7 +40,7 @@ public:
 		TArray<int32> VertexIndicesIn;
 
 	UPROPERTY(meta = (DataflowInput, DisplayName = "BoneIndex"))
-		int32 BoneIndexIn = INDEX_NONE;
+		int32 BoneIndexIn = 0;
 
 	UPROPERTY(meta = (DataflowOutput, DisplayName = "TargetIndices"))
 		TArray<int32> TargetIndicesOut;
diff --git a/Engine/Plugins/Experimental/ChaosFlesh/Source/ChaosFleshNodes/Public/Dataflow/ChaosFleshSkeletalBindingsNode.h b/Engine/Plugins/Experimental/ChaosFlesh/Source/ChaosFleshNodes/Public/Dataflow/ChaosFleshSkeletalBindingsNode.h
new file mode 100644
--- /dev/null
+++ b/Engine/Plugins/Experimental/ChaosFlesh/Source/ChaosFleshNodes/Public/Dataflow/ChaosFleshSkeletalBindingsNode.h
@@ -0,0 +1,53 @@
+// Copyright Epic Games, Inc. All Rights Reserved.
+
+#pragma once 
+
+#include "CoreMinimal.h"
+#include "Dataflow/DataflowEngine.h"
+#include "GeometryCollection/ManagedArrayCollection.h"
+
+#include "ChaosFleshSkeletalBindingsNode.generated.h"
+
+class UStaticMesh;
+class USkeletalMesh;
+
+DECLARE_LOG_CATEGORY_EXTERN(LogSkeletalBindings, Verbose, All);
+
+// Generate barycentric bindings (used by the FleshDeformer deformer graph) of a render surface to a tetrahedral mesh.
+USTRUCT(meta = (DataflowFlesh))
+struct FGenerateSkeletalBindings : public FDataflowNode
+{
+	GENERATED_USTRUCT_BODY()
+	DATAFLOW_NODE_DEFINE_INTERNAL(FGenerateSkeletalBindings, "GenerateSkeletalBindings", "Flesh", "")
+	DATAFLOW_NODE_RENDER_TYPE(FGeometryCollection::StaticType(), "Collection")
+
+public:
+	typedef FManagedArrayCollection DataType;
+
+	// Passthrough geometry collection. Bindings are stored as standalone groups in the \p Collection, keyed by the name of the input render mesh and all available LOD's.
+	UPROPERTY(meta = (DataflowInput, DataflowOutput, DisplayName = "Collection", DataflowPassthrough = "Collection"))
+	FManagedArrayCollection Collection;
+
+	UPROPERTY(meta = (DataflowInput, DisplayName = "BoneIndex"))
+	int32 BoneIndexIn = 0;
+
+	// The input mesh, whose render surface is used to generate bindings.
+	UPROPERTY(EditAnywhere, Category = "Dataflow", meta = (DataflowInput, DisplayName = "SkeletalMesh"))
+	TObjectPtr<const USkeletalMesh> SkeletalMeshIn = nullptr;
+
+	FGenerateSkeletalBindings(const Dataflow::FNodeParameters& InParam, FGuid InGuid = FGuid::NewGuid())
+		: FDataflowNode(InParam, InGuid)
+	{
+		RegisterInputConnection(&Collection);
+		RegisterOutputConnection(&Collection, &Collection);
+		RegisterInputConnection(&BoneIndexIn);
+		RegisterInputConnection(&SkeletalMeshIn);
+	}
+
+	virtual void Evaluate(Dataflow::FContext& Context, const FDataflowOutput* Out) const override;
+};
+
+namespace Dataflow
+{
+	void ChaosFleshSkeletalBindingsNode();
+}
diff --git a/Engine/Plugins/Experimental/Dataflow/Source/DataflowEditor/Private/Dataflow/DataflowEditorToolkit.cpp b/Engine/Plugins/Experimental/Dataflow/Source/DataflowEditor/Private/Dataflow/DataflowEditorToolkit.cpp
--- a/Engine/Plugins/Experimental/Dataflow/Source/DataflowEditor/Private/Dataflow/DataflowEditorToolkit.cpp
+++ b/Engine/Plugins/Experimental/Dataflow/Source/DataflowEditor/Private/Dataflow/DataflowEditorToolkit.cpp
@@ -288,6 +288,7 @@ TSharedRef<SDataflowGraphEditor> FDataflowEditorToolkit::CreateGraphEditorWidget
 		{
 			Context = TSharedPtr< Dataflow::FEngineContext>(new Dataflow::FAssetContext(Asset, Dataflow, Dataflow::FTimestamp::Invalid));
 		}
+		Node->Invalidate();
 		LastNodeTimestamp = Dataflow::FTimestamp::Invalid;
 
 		FDataflowEditorCommands::EvaluateTerminalNode(*Context.Get(), LastNodeTimestamp, Dataflow, Node, Out, Asset, TerminalPath);
diff --git a/Engine/Source/Runtime/Experimental/Chaos/Private/GeometryCollection/Facades/CollectionTetrahedralSkeletalBindingsFacade.cpp b/Engine/Source/Runtime/Experimental/Chaos/Private/GeometryCollection/Facades/CollectionTetrahedralSkeletalBindingsFacade.cpp
new file mode 100644
--- /dev/null
+++ b/Engine/Source/Runtime/Experimental/Chaos/Private/GeometryCollection/Facades/CollectionTetrahedralSkeletalBindingsFacade.cpp
@@ -0,0 +1,147 @@
+// Copyright Epic Games, Inc. All Rights Reserved.
+
+/*=============================================================================
+	GeometryCollection.cpp: FGeometryCollection methods.
+=============================================================================*/
+
+#include "GeometryCollection/Facades/CollectionTetrahedralSkeletalBindingsFacade.h"
+#include "GeometryCollection/GeometryCollection.h"
+#include "GeometryCollection/ManagedArrayAccessor.h"
+
+
+namespace GeometryCollection::Facades
+{
+	// Groups 
+	const FName FTetrahedralSkeletalBindings::MeshBindingsGroupName = "SkeletalBindings";
+	const FName FTetrahedralSkeletalBindings::MeshBindingsIdGroupName = "SkeletalBindingsId";
+
+	// Attributes
+	const FName FTetrahedralSkeletalBindings::MeshIdAttributeName = "MeshId";
+	const FName FTetrahedralSkeletalBindings::MeshIndexAttributeName = "MeshIdIndex";
+	const FName FTetrahedralSkeletalBindings::TetrahedronIndexAttributeName = "TetrahedronIndex";
+	const FName FTetrahedralSkeletalBindings::WeightsAttributeName = "Weights";
+	const FName FTetrahedralSkeletalBindings::SkeletalIndexAttributeName = "SkeletalIndex";
+
+	FTetrahedralSkeletalBindings::FTetrahedralSkeletalBindings(FManagedArrayCollection& InCollection)
+		: MeshIdAttribute(InCollection, MeshIdAttributeName, MeshBindingsIdGroupName)
+		, MeshGroupIndexAttribute(InCollection, MeshIndexAttributeName, MeshBindingsGroupName, MeshBindingsIdGroupName)
+		, TetrahedronIndexAttribute(InCollection, TetrahedronIndexAttributeName, MeshBindingsGroupName, "Tetrahedral")
+		, WeightsAttribute(InCollection, WeightsAttributeName, MeshBindingsGroupName)
+		, SkeletalIndexAttribute(InCollection, SkeletalIndexAttributeName, MeshBindingsGroupName)
+		, TetrahedronAttribute(InCollection, "Tetrahedron", "Tetrahedral")
+		, VerticesAttribute(InCollection, "Vertex", FGeometryCollection::VerticesGroup)
+	{
+		DefineSchema();
+	}
+
+	FTetrahedralSkeletalBindings::FTetrahedralSkeletalBindings(const FManagedArrayCollection& InCollection)
+		: MeshIdAttribute(InCollection, MeshIdAttributeName, MeshBindingsIdGroupName)
+		, MeshGroupIndexAttribute(InCollection, MeshIndexAttributeName, MeshBindingsGroupName, MeshBindingsIdGroupName)
+		, TetrahedronIndexAttribute(InCollection, TetrahedronIndexAttributeName, MeshBindingsGroupName, "Tetrahedral")
+		, WeightsAttribute(InCollection, WeightsAttributeName, MeshBindingsGroupName)
+		, SkeletalIndexAttribute(InCollection, SkeletalIndexAttributeName, MeshBindingsGroupName) 
+		, TetrahedronAttribute(InCollection, "Tetrahedron", "Tetrahedral")
+		, VerticesAttribute(InCollection, "Vertex", FGeometryCollection::VerticesGroup)
+	{}
+
+	FTetrahedralSkeletalBindings::~FTetrahedralSkeletalBindings()
+	{}
+
+	void FTetrahedralSkeletalBindings::DefineSchema()
+	{
+		check(!IsConst());
+		MeshIdAttribute.Add();
+		MeshGroupIndexAttribute.Add();
+		TetrahedronIndexAttribute.Add();
+		WeightsAttribute.Add();
+		SkeletalIndexAttribute.Add();
+	}
+
+	bool FTetrahedralSkeletalBindings::IsValid() const
+	{
+		return MeshIdAttribute.IsValid() && MeshGroupIndexAttribute.IsValid() && TetrahedronIndexAttribute.IsValid()
+			&& WeightsAttribute.IsValid() && SkeletalIndexAttribute.IsValid() && TetrahedronAttribute.IsValid()
+			&& VerticesAttribute.IsValid();
+	}
+
+	FString FTetrahedralSkeletalBindings::GenerateMeshGroupName( const int32 TetMeshIdx,const FName& MeshId,const int32 LOD)
+	{
+		return FString::Printf(TEXT("v1:%d:%s:%d"), TetMeshIdx, *MeshId.ToString(), LOD);
+	}
+
+	void FTetrahedralSkeletalBindings::SetBindings(const FString& InMeshGroupIndex, const TArray<int32>& InTetrahedronIndex,
+												   const TArray<FVector4f>& WeightsIn, const TArray<int32>& InSkeletalIndex)
+	{
+		check(!IsConst());
+		if (InTetrahedronIndex.Num() && InTetrahedronIndex.Num() == WeightsIn.Num() && InTetrahedronIndex.Num() == InSkeletalIndex.Num())
+		{
+			TManagedArray<FString>& MeshId = MeshIdAttribute.Modify();
+			TManagedArray<int32>& MeshGroupIndex = MeshGroupIndexAttribute.Modify();
+			TManagedArray<int32>& TetrahedronIndex = TetrahedronIndexAttribute.Modify();
+			TManagedArray<FVector4f>& Weights = WeightsAttribute.Modify();
+			TManagedArray<int32>& SkeletalIndex = SkeletalIndexAttribute.Modify();
+
+			int32 GroupIndex = MeshId.Find(InMeshGroupIndex);
+			if (GroupIndex == INDEX_NONE)
+			{
+				GroupIndex = MeshIdAttribute.AddElements(1);
+			}
+			MeshId[GroupIndex] = InMeshGroupIndex;
+
+			int32 NumTetrahedron = InTetrahedronIndex.Num();
+			int32 StartIndex = MeshGroupIndexAttribute.AddElements(NumTetrahedron);
+			for (int32 CurrentIndex = StartIndex, i = 0; CurrentIndex < StartIndex + NumTetrahedron; CurrentIndex++, i++)
+			{
+				MeshGroupIndex[CurrentIndex] = GroupIndex;
+				TetrahedronIndex[CurrentIndex] = InTetrahedronIndex[i];
+				Weights[CurrentIndex] = WeightsIn[i];
+				SkeletalIndex[CurrentIndex] = InSkeletalIndex[i];
+			}
+		}
+	}
+
+
+	void FTetrahedralSkeletalBindings::CalculateBindings(const FString & InKey, const TArray<FVector3f>& InVertices, TArray<FVector>& OutPosition, TArray<bool>* OutInfluence) const
+	{
+		auto UEVert3d = [](FVector3f V) { return FVector3d(V.X, V.Y, V.Z); };
+
+		if (IsValid())
+		{
+			int32 GroupIndex = MeshIdAttribute.Get().Find(InKey);
+			if (GroupIndex != INDEX_NONE)
+			{
+				const TManagedArray<int32>& MeshGroupIndex = MeshGroupIndexAttribute.Get();
+				const TManagedArray<int32>& TetrahedronIndex = TetrahedronIndexAttribute.Get();
+				const TManagedArray<FVector4f>& Weights = WeightsAttribute.Get();
+				const TManagedArray<int32>& SkeletalIndices = SkeletalIndexAttribute.Get();
+				const TManagedArray<FIntVector4>& Tetrahedron = TetrahedronAttribute.Get();
+
+				int32 NumBindings = MeshGroupIndex.Num();
+				for (int32 BindingIndex = 0; BindingIndex < MeshGroupIndex.Num(); BindingIndex++)
+				{
+					if (MeshGroupIndex[BindingIndex] == GroupIndex)
+					{
+						int32 TetIndex = TetrahedronIndex[BindingIndex];
+						if (0 <= TetIndex && TetIndex < TetrahedronAttribute.Num())
+						{
+							const FIntVector4& Tet = Tetrahedron[TetIndex];
+							int32 SkeletonIndex = SkeletalIndices[BindingIndex];
+							if (0 <= SkeletonIndex && SkeletonIndex < OutPosition.Num())
+							{
+								const FVector4 Weight(Weights[BindingIndex][0], Weights[BindingIndex][1], Weights[BindingIndex][2], Weights[BindingIndex][3]);
+								FVector X0 = UEVert3d(InVertices[Tet[0]]);
+								FVector X1 = UEVert3d(InVertices[Tet[1]]);
+								FVector X2 = UEVert3d(InVertices[Tet[2]]);
+								FVector X3 = UEVert3d(InVertices[Tet[3]]);
+								OutPosition[SkeletonIndex] = Weight[0] * X0 + Weight[1] * X1 + Weight[2] * X2 + Weight[3] * X3;
+								if (OutInfluence!=nullptr) if(SkeletonIndex<OutInfluence->Num()) (*OutInfluence)[SkeletonIndex] = true;
+							}
+						}
+					}
+				}
+			}
+		}
+	}
+};
+
+
diff --git a/Engine/Source/Runtime/Experimental/Chaos/Public/GeometryCollection/Facades/CollectionTetrahedralSkeletalBindingsFacade.h b/Engine/Source/Runtime/Experimental/Chaos/Public/GeometryCollection/Facades/CollectionTetrahedralSkeletalBindingsFacade.h
new file mode 100644
--- /dev/null
+++ b/Engine/Source/Runtime/Experimental/Chaos/Public/GeometryCollection/Facades/CollectionTetrahedralSkeletalBindingsFacade.h
@@ -0,0 +1,75 @@
+// Copyright Epic Games, Inc. All Rights Reserved.
+
+#pragma once
+#include "Containers/UnrealString.h"
+#include "Containers/Array.h"
+#include "GeometryCollection/ManagedArrayCollection.h"
+#include "GeometryCollection/ManagedArrayAccessor.h"
+
+namespace GeometryCollection::Facades
+{
+
+	/**
+	* FTetrahedralSkeletalBindings
+	* 
+	* Interface for storing and retrieving bindings of transforms of a SkeletalMesh to
+	* tetrahedral meshes. All bindings are stored in a single group, and a lookup table
+	* is aviable to match bindings to specifi meshs by name. 
+	*/
+	class CHAOS_API FTetrahedralSkeletalBindings
+ 	{
+	public:
+		// Groups 
+		static const FName MeshBindingsGroupName;
+		static const FName MeshBindingsIdGroupName;
+
+		// Attributes
+		static const FName MeshIdAttributeName;
+		static const FName MeshIndexAttributeName;
+		static const FName TetrahedronIndexAttributeName;
+		static const FName WeightsAttributeName;
+		static const FName SkeletalIndexAttributeName;
+
+		/**
+		* FSelectionFacade Constuctor
+		* @param VertixDependencyGroup : GroupName the index attribute is dependent on. 
+		*/
+		FTetrahedralSkeletalBindings(FManagedArrayCollection& InSelf);
+		FTetrahedralSkeletalBindings(const FManagedArrayCollection& InSelf);
+		virtual ~FTetrahedralSkeletalBindings();
+
+		/** 
+		* Create the facade schema. 
+		*/
+		void DefineSchema();
+
+		/** Returns \c true if the facade is operating on a read-only geometry collection. */
+		bool IsConst() const { return MeshIdAttribute.IsConst(); }
+
+		/** 
+		* Returns \c true if the Facade defined on the collection, and is initialized to
+		* a valid bindings group.
+		*/
+		bool IsValid() const;
+
+		//
+		// API
+		//
+		static FString GenerateMeshGroupName(const int32 TetMeshIdx, const FName& SkeletalMeshName, const int32 LOD = INDEX_NONE);
+		void SetBindings(const FString& InMeshGroupIndex, const TArray<int32>& InTetrahedronIndex, const TArray<FVector4f>& WeightsIn, const TArray<int32>& InSkeletalIndex);
+		void CalculateBindings(const FString& InKey, const TArray<FVector3f>& InVertices, TArray<FVector>& OutPosition, TArray<bool>* OutInfluence = nullptr) const;
+		
+	private:
+		TManagedArrayAccessor<FString> MeshIdAttribute;
+
+		TManagedArrayAccessor<int32> MeshGroupIndexAttribute;
+		TManagedArrayAccessor<int32> TetrahedronIndexAttribute;
+		TManagedArrayAccessor<FVector4f> WeightsAttribute;
+		TManagedArrayAccessor<int32> SkeletalIndexAttribute;
+
+		TManagedArrayAccessor<FIntVector4> TetrahedronAttribute;
+		TManagedArrayAccessor<FVector3f> VerticesAttribute;
+
+	};
+
+}
