ChaosFlesh : Collisions Manager
-Added Collision components to the solver.
-Added collision body support to the solver level API.
-UFleshComponet to respect IsVisible
-UDeformablePhysicsComponent now derives from UPrimitiveComponent.

diff --git a/Engine/Plugins/Experimental/ChaosFlesh/Source/ChaosFleshEngine/Private/ChaosFlesh/ChaosDeformableCollisionsActor.cpp b/Engine/Plugins/Experimental/ChaosFlesh/Source/ChaosFleshEngine/Private/ChaosFlesh/ChaosDeformableCollisionsActor.cpp
new file mode 100644
--- /dev/null
+++ b/Engine/Plugins/Experimental/ChaosFlesh/Source/ChaosFleshEngine/Private/ChaosFlesh/ChaosDeformableCollisionsActor.cpp
@@ -0,0 +1,106 @@
+// Copyright Epic Games, Inc. All Rights Reserved.
+
+#include "ChaosFlesh/ChaosDeformableCollisionsActor.h"
+
+#include "ChaosFlesh/ChaosDeformableSolverComponent.h"
+#include "ChaosFlesh/ChaosDeformableSolverActor.h"
+#include "Engine/StaticMeshActor.h"
+
+#include UE_INLINE_GENERATED_CPP_BY_NAME(ChaosDeformableCollisionsActor)
+
+ADeformableCollisionsActor::ADeformableCollisionsActor(const FObjectInitializer& ObjectInitializer)
+	: Super(ObjectInitializer)
+{
+	DeformableCollisionsComponent = CreateDefaultSubobject<UDeformableCollisionsComponent>(TEXT("DeformableCollisionsComponent"));
+	RootComponent = DeformableCollisionsComponent;
+	PrimaryActorTick.bCanEverTick = false;
+}
+
+void ADeformableCollisionsActor::EnableSimulation(ADeformableSolverActor* InActor)
+{
+	if (InActor)
+	{
+		if (DeformableCollisionsComponent)
+		{
+			DeformableCollisionsComponent->EnableSimulation(InActor->GetDeformableSolverComponent());
+		}
+	}
+}
+
+
+#if WITH_EDITOR
+void ADeformableCollisionsActor::PreEditChange(FProperty* PropertyThatWillChange)
+{
+	Super::PreEditChange(PropertyThatWillChange);
+	if (PropertyThatWillChange && PropertyThatWillChange->GetFName() == GET_MEMBER_NAME_CHECKED(ADeformableCollisionsActor, PrimarySolver))
+	{
+		PreEditChangePrimarySolver = PrimarySolver;
+	}
+	if (PropertyThatWillChange && PropertyThatWillChange->GetFName() == GET_MEMBER_NAME_CHECKED(ADeformableCollisionsActor, StaticCollisions))
+	{
+		PreEditChangeCollisionBodies = StaticCollisions;
+	}
+}
+
+
+void ADeformableCollisionsActor::PostEditChangeProperty(struct FPropertyChangedEvent& PropertyChangedEvent)
+{
+	Super::PostEditChangeProperty(PropertyChangedEvent);
+
+	if (PropertyChangedEvent.Property && PropertyChangedEvent.Property->GetFName() == GET_MEMBER_NAME_CHECKED(ADeformableCollisionsActor, StaticCollisions))
+	{
+		typedef TObjectPtr<AStaticMeshActor> FActorType;
+		TSet< FActorType>  A(PreEditChangeCollisionBodies), B(StaticCollisions);
+		for (auto AddedActor : B.Difference(A).Array())
+		{
+			if (AddedActor && AddedActor->GetStaticMeshComponent())
+			{
+				GetCollisionsComponent()->AddStaticMeshComponent(AddedActor->GetStaticMeshComponent());
+			}
+		}
+
+		for (auto RemovedActor : A.Difference(B).Array())
+		{
+			if (RemovedActor && RemovedActor->GetStaticMeshComponent())
+			{
+				GetCollisionsComponent()->RemoveStaticMeshComponent(RemovedActor->GetStaticMeshComponent());
+			}
+		}
+
+		PreEditChangeCollisionBodies.Empty();
+	}
+	if (PropertyChangedEvent.Property && PropertyChangedEvent.Property->GetFName() == GET_MEMBER_NAME_CHECKED(ADeformableCollisionsActor, PrimarySolver))
+	{
+		if (PrimarySolver)
+		{
+			if (UDeformableSolverComponent* SolverComponent = PrimarySolver->GetDeformableSolverComponent())
+			{
+				if (DeformableCollisionsComponent)
+				{
+					DeformableCollisionsComponent->PrimarySolverComponent = SolverComponent;
+					if (!SolverComponent->DeformableComponents.Contains(DeformableCollisionsComponent))
+					{
+						SolverComponent->DeformableComponents.Add(TObjectPtr<UDeformablePhysicsComponent>(DeformableCollisionsComponent));
+					}
+				}
+			}
+		}
+		else if (PreEditChangePrimarySolver)
+		{
+			if (UDeformableSolverComponent* SolverComponent = PreEditChangePrimarySolver->GetDeformableSolverComponent())
+			{
+				if (DeformableCollisionsComponent)
+				{
+					DeformableCollisionsComponent->PrimarySolverComponent = nullptr;
+					if (SolverComponent->DeformableComponents.Contains(DeformableCollisionsComponent))
+					{
+						SolverComponent->DeformableComponents.Remove(DeformableCollisionsComponent);
+					}
+				}
+			}
+		}
+	}
+}
+#endif
+
+
diff --git a/Engine/Plugins/Experimental/ChaosFlesh/Source/ChaosFleshEngine/Private/ChaosFlesh/ChaosDeformableCollisionsComponent.cpp b/Engine/Plugins/Experimental/ChaosFlesh/Source/ChaosFleshEngine/Private/ChaosFlesh/ChaosDeformableCollisionsComponent.cpp
new file mode 100644
--- /dev/null
+++ b/Engine/Plugins/Experimental/ChaosFlesh/Source/ChaosFleshEngine/Private/ChaosFlesh/ChaosDeformableCollisionsComponent.cpp
@@ -0,0 +1,155 @@
+// Copyright Epic Games, Inc. All Rights Reserved.
+
+#include "ChaosFlesh/ChaosDeformableCollisionsComponent.h"
+
+#include "ChaosFlesh/ChaosDeformableSolverComponent.h"
+#include "Chaos/Deformable/ChaosDeformableCollisionsProxy.h"
+#include "Chaos/Sphere.h"
+#include "Chaos/Convex.h"
+#include "Chaos/ImplicitObject.h"
+#include "Chaos/ImplicitObjectScaled.h"
+#include "Chaos/TriangleMeshImplicitObject.h"
+#include "Components/StaticMeshComponent.h"
+#include "PhysicsEngine/BodySetup.h"
+DEFINE_LOG_CATEGORY_STATIC(LogUDeformableCollisionsComponentInternal, Log, All);
+
+UDeformableCollisionsComponent::UDeformableCollisionsComponent(const FObjectInitializer& ObjectInitializer)
+	: Super(ObjectInitializer)
+{
+	PrimaryComponentTick.bCanEverTick = false;
+	bTickInEditor = false;
+}
+
+void UDeformableCollisionsComponent::AddStaticMeshComponent(UStaticMeshComponent* StaticMeshComponent)
+{
+	if (StaticMeshComponent)
+	{
+		if (!CollisionBodies.Contains(StaticMeshComponent))
+		{
+			CollisionBodies.Add(TObjectPtr<UStaticMeshComponent>(StaticMeshComponent) );
+			AddedBodies.Add(StaticMeshComponent);
+		}
+	}
+}
+
+void UDeformableCollisionsComponent::RemoveStaticMeshComponent(UStaticMeshComponent* StaticMeshComponent)
+{
+	if (StaticMeshComponent)
+	{
+		if (CollisionBodies.Contains(StaticMeshComponent))
+		{
+			CollisionBodies.Remove(TObjectPtr<UStaticMeshComponent>(StaticMeshComponent));
+			RemovedBodies.Add(StaticMeshComponent);
+		}
+	}
+}
+
+
+UDeformablePhysicsComponent::FThreadingProxy* UDeformableCollisionsComponent::NewProxy()
+{
+	for (auto& Body : CollisionBodies)
+	{
+		if (Body)
+		{
+			if (!AddedBodies.Contains(Body))
+			{
+				AddedBodies.Add(Body);
+			}
+		}
+	}
+	return new FCollisionThreadingProxy(this);
+}
+
+template<> class Chaos::TSphere<float, 3>;
+
+UDeformableCollisionsComponent::FDataMapValue 
+UDeformableCollisionsComponent::NewDeformableData()
+{ 
+	TArray < Chaos::Softs::FCollisionObjectAddedBodies> AddedBodiesData;
+	TArray < Chaos::Softs::FCollisionObjectRemovedBodies> RemovedBodiesData;
+	TArray < Chaos::Softs::FCollisionObjectUpdatedBodies> UpdateBodiesData;
+
+	for (auto& CollisionBody : AddedBodies)
+	{
+		if( CollisionBody )
+		{
+			if (UStaticMesh* StaticMesh = CollisionBody->GetStaticMesh())
+			{
+				if (UBodySetup* BodySetup = StaticMesh->GetBodySetup())
+				{
+					using namespace Chaos;
+					FImplicitObject* Geometry = nullptr;
+					FTransform P = CollisionBody->GetComponentToWorld();
+					FVector Scale = P.GetScale3D();
+					P.RemoveScaling();
+
+					if (BodySetup->AggGeom.SphereElems.Num())
+					{
+						FKSphereElem& S = BodySetup->AggGeom.SphereElems[0];
+						Geometry = new ::Chaos::TSphere<Chaos::FReal, 3>(S.Center, S.Radius * Scale.GetMax());
+					}
+					else if (BodySetup->AggGeom.BoxElems.Num())
+					{
+						FKBoxElem& B = BodySetup->AggGeom.BoxElems[0];
+						Geometry = new TBox<Chaos::FReal, 3>(
+							B.Center - 0.5 * Scale * FVector(B.X, B.Y, B.Z),
+							B.Center + 0.5 * Scale * FVector(B.X, B.Y, B.Z));
+					}
+					else if (BodySetup->AggGeom.ConvexElems.Num())
+					{
+						FKConvexElem& C = BodySetup->AggGeom.ConvexElems[0];
+						if (C.VertexData.Num())
+						{
+							TSharedPtr<Chaos::FConvex, ESPMode::ThreadSafe> ConvexMesh = C.GetChaosConvexMesh();	
+							TArray<FConvex::FVec3Type> Vertices;
+							Vertices.SetNum(C.VertexData.Num());
+							for (int i = 0; i < Vertices.Num(); i++)
+							{
+								Vertices[i] = C.VertexData[i] * Scale;
+							}
+							Geometry = new FConvex(Vertices, 0.f);
+						}
+					}
+					//else if (BodySetup->ChaosTriMeshes.Num())
+					//{
+					//	TArray<TSharedPtr<Chaos::FTriangleMeshImplicitObject, ESPMode::ThreadSafe>> ChaosTriMeshes;
+					//	Geometry = BodySetup->ChaosTriMeshes[0]->CopySlow().Release();
+					//}
+					if (Geometry)
+					{
+						AddedBodiesData.Add(Chaos::Softs::FCollisionObjectAddedBodies(CollisionBody, P, "", Geometry));
+					}
+				}
+			}
+		}
+	}
+
+
+	for (auto& RemovedBody : RemovedBodies) { RemovedBodiesData.Add({RemovedBody}); }
+	AddedBodies.Empty();
+	RemovedBodies.Empty();
+
+	for (auto& CollisionBody : CollisionBodies) 
+	{
+		if (CollisionBody)
+		{
+			FTransform P = CollisionBody->GetComponentToWorld();
+			UpdateBodiesData.Add({ CollisionBody,P });
+		}
+	}
+
+
+	return FDataMapValue(new Chaos::Softs::FCollisionManagerProxy::FCollisionsInputBuffer(
+		AddedBodiesData, RemovedBodiesData, UpdateBodiesData, this));
+}
+
+
+
+
+
+
+
+
+
+
+     
diff --git a/Engine/Plugins/Experimental/ChaosFlesh/Source/ChaosFleshEngine/Private/ChaosFlesh/ChaosDeformableSolverComponent.cpp b/Engine/Plugins/Experimental/ChaosFlesh/Source/ChaosFleshEngine/Private/ChaosFlesh/ChaosDeformableSolverComponent.cpp
--- a/Engine/Plugins/Experimental/ChaosFlesh/Source/ChaosFleshEngine/Private/ChaosFlesh/ChaosDeformableSolverComponent.cpp
+++ b/Engine/Plugins/Experimental/ChaosFlesh/Source/ChaosFleshEngine/Private/ChaosFlesh/ChaosDeformableSolverComponent.cpp
@@ -2,6 +2,7 @@
 
 #include "ChaosFlesh/ChaosDeformableSolverComponent.h"
 
+#include "ChaosFlesh/ChaosDeformableCollisionsComponent.h"
 #include "ChaosFlesh/ChaosDeformablePhysicsComponent.h"
 #include "ChaosFlesh/ChaosDeformableSolverActor.h"
 #include "Engine/World.h"
@@ -95,7 +96,7 @@ bool UDeformableSolverComponent::IsSimulating(UDeformablePhysicsComponent* InCom
 	if (InComponent)
 	{
 		const UDeformableSolverComponent* ComponentSolver = InComponent->PrimarySolverComponent.Get();
-		return ComponentSolver ==this;
+		return (void*)ComponentSolver == (void*)this;
 	}
 	return false;
 }
diff --git a/Engine/Plugins/Experimental/ChaosFlesh/Source/ChaosFleshEngine/Private/ChaosFlesh/FleshComponent.cpp b/Engine/Plugins/Experimental/ChaosFlesh/Source/ChaosFleshEngine/Private/ChaosFlesh/FleshComponent.cpp
--- a/Engine/Plugins/Experimental/ChaosFlesh/Source/ChaosFleshEngine/Private/ChaosFlesh/FleshComponent.cpp
+++ b/Engine/Plugins/Experimental/ChaosFlesh/Source/ChaosFleshEngine/Private/ChaosFlesh/FleshComponent.cpp
@@ -299,105 +299,111 @@ void UFleshComponent::RenderProceduralMesh()
 	bool bCanRender = false;
 	if (const UFleshAsset* FleshAsset = GetRestCollection())
 	{
+		if (IsVisible())
+		{
 #if WITH_EDITORONLY_DATA
-		if (Mesh && FleshAsset->bRenderInEditor && CVarParams.bDoDrawSimulationMesh)
+			if (FleshAsset->bRenderInEditor)
 #endif
-		{
-			if (const FFleshCollection* Flesh = FleshAsset->GetCollection())
 			{
-				int32 NumVertices = Flesh->NumElements(FGeometryCollection::VerticesGroup);
-				int32 NumFaces = Flesh->NumElements(FGeometryCollection::FacesGroup);
-				if (NumFaces && NumVertices)
+				if (Mesh && CVarParams.bDoDrawSimulationMesh)
 				{
-					if (RenderMesh && RenderMesh->Vertices.Num() != NumFaces * 3)
+					if (const FFleshCollection* Flesh = FleshAsset->GetCollection())
 					{
-						ResetProceduralMesh();
-					}
+						int32 NumVertices = Flesh->NumElements(FGeometryCollection::VerticesGroup);
+						int32 NumFaces = Flesh->NumElements(FGeometryCollection::FacesGroup);
+						if (NumFaces && NumVertices)
+						{
+							if (RenderMesh && RenderMesh->Vertices.Num() != NumFaces * 3)
+							{
+								ResetProceduralMesh();
+							}
 
-					if (!RenderMesh)
-					{
-						RenderMesh = new FFleshRenderMesh;
+							if (!RenderMesh)
+							{
+								RenderMesh = new FFleshRenderMesh;
 
-						for (int i = 0; i < NumFaces; ++i)
-						{
-							const auto& P1 = Flesh->Vertex[Flesh->Indices[i][0]];
-							const auto& P2 = Flesh->Vertex[Flesh->Indices[i][1]];
-							const auto& P3 = Flesh->Vertex[Flesh->Indices[i][2]];
-
-							RenderMesh->Vertices.Add(FVector(P1));
-							RenderMesh->Vertices.Add(FVector(P2));
-							RenderMesh->Vertices.Add(FVector(P3));
-
-							RenderMesh->Colors.Add(FLinearColor::White);
-							RenderMesh->Colors.Add(FLinearColor::White);
-							RenderMesh->Colors.Add(FLinearColor::White);
-
-							RenderMesh->UVs.Add(FVector2D(0, 0));
-							RenderMesh->UVs.Add(FVector2D(0, 0));
-							RenderMesh->UVs.Add(FVector2D(0, 0));
-
-							RenderMesh->Triangles.Add(3 * i);
-							RenderMesh->Triangles.Add(3 * i + 1);
-							RenderMesh->Triangles.Add(3 * i + 2);
-
-							auto Normal = -Chaos::FVec3::CrossProduct(P3 - P1, P2 - P1);
-							RenderMesh->Normals.Add(Normal);
-							RenderMesh->Normals.Add(Normal);
-							RenderMesh->Normals.Add(Normal);
-
-							auto Tangent = (P2 - P1).GetSafeNormal();
-							RenderMesh->Tangents.Add(FProcMeshTangent(Tangent[0], Tangent[1], Tangent[2]));
-							Tangent = (P3 - P2).GetSafeNormal();
-							RenderMesh->Tangents.Add(FProcMeshTangent(Tangent[0], Tangent[1], Tangent[2]));
-							Tangent = (P1 - P3).GetSafeNormal();
-							RenderMesh->Tangents.Add(FProcMeshTangent(Tangent[0], Tangent[1], Tangent[2]));
-						}
+								for (int i = 0; i < NumFaces; ++i)
+								{
+									const auto& P1 = Flesh->Vertex[Flesh->Indices[i][0]];
+									const auto& P2 = Flesh->Vertex[Flesh->Indices[i][1]];
+									const auto& P3 = Flesh->Vertex[Flesh->Indices[i][2]];
+
+									RenderMesh->Vertices.Add(FVector(P1));
+									RenderMesh->Vertices.Add(FVector(P2));
+									RenderMesh->Vertices.Add(FVector(P3));
+
+									RenderMesh->Colors.Add(FLinearColor::White);
+									RenderMesh->Colors.Add(FLinearColor::White);
+									RenderMesh->Colors.Add(FLinearColor::White);
+
+									RenderMesh->UVs.Add(FVector2D(0, 0));
+									RenderMesh->UVs.Add(FVector2D(0, 0));
+									RenderMesh->UVs.Add(FVector2D(0, 0));
+
+									RenderMesh->Triangles.Add(3 * i);
+									RenderMesh->Triangles.Add(3 * i + 1);
+									RenderMesh->Triangles.Add(3 * i + 2);
+
+									auto Normal = -Chaos::FVec3::CrossProduct(P3 - P1, P2 - P1);
+									RenderMesh->Normals.Add(Normal);
+									RenderMesh->Normals.Add(Normal);
+									RenderMesh->Normals.Add(Normal);
+
+									auto Tangent = (P2 - P1).GetSafeNormal();
+									RenderMesh->Tangents.Add(FProcMeshTangent(Tangent[0], Tangent[1], Tangent[2]));
+									Tangent = (P3 - P2).GetSafeNormal();
+									RenderMesh->Tangents.Add(FProcMeshTangent(Tangent[0], Tangent[1], Tangent[2]));
+									Tangent = (P1 - P3).GetSafeNormal();
+									RenderMesh->Tangents.Add(FProcMeshTangent(Tangent[0], Tangent[1], Tangent[2]));
+								}
 
-						Mesh->SetRelativeTransform(GetComponentTransform());
-						Mesh->CreateMeshSection_LinearColor(0, RenderMesh->Vertices, RenderMesh->Triangles, RenderMesh->Normals, RenderMesh->UVs, RenderMesh->Colors, RenderMesh->Tangents, false);
-					}
-					else
-					{
+								Mesh->SetRelativeTransform(GetComponentTransform());
+								Mesh->CreateMeshSection_LinearColor(0, RenderMesh->Vertices, RenderMesh->Triangles, RenderMesh->Normals, RenderMesh->UVs, RenderMesh->Colors, RenderMesh->Tangents, false);
+							}
+							else
+							{
 
-						const TManagedArray<FVector3f>* RenderVertex = &Flesh->Vertex;
-						if (GetDynamicCollection())
-						{
-							const TManagedArray<FVector3f>& DynamicVertex = GetDynamicCollection()->GetPositions();
-							if (DynamicVertex.Num()) RenderVertex = &DynamicVertex;
-						}
-						auto InRange = [](int32 Size, int32 Val) { return 0 <= Val && Val < Size; };
+								const TManagedArray<FVector3f>* RenderVertex = &Flesh->Vertex;
+								if (GetDynamicCollection())
+								{
+									const TManagedArray<FVector3f>& DynamicVertex = GetDynamicCollection()->GetPositions();
+									if (DynamicVertex.Num()) RenderVertex = &DynamicVertex;
+								}
+								auto InRange = [](int32 Size, int32 Val) { return 0 <= Val && Val < Size; };
 
-						// Display only
-						for (int i = 0; i < NumFaces; ++i)
-						{
-							const auto& P1 = (*RenderVertex)[Flesh->Indices[i][0]];
-							const auto& P2 = (*RenderVertex)[Flesh->Indices[i][1]];
-							const auto& P3 = (*RenderVertex)[Flesh->Indices[i][2]];
-
-							RenderMesh->Vertices[3 * i] = FVector(P1);
-							RenderMesh->Vertices[3 * i + 1] = FVector(P2);
-							RenderMesh->Vertices[3 * i + 2] = FVector(P3);
-
-							auto Normal = Chaos::FVec3::CrossProduct(P3 - P1, P2 - P1);
-							RenderMesh->Normals[3 * i] = Normal;
-							RenderMesh->Normals[3 * i + 1] = Normal;
-							RenderMesh->Normals[3 * i + 2] = Normal;
-
-							auto Tangent = (P2 - P1).GetSafeNormal();
-							RenderMesh->Tangents[3 * i] = FProcMeshTangent(Tangent[0], Tangent[1], Tangent[2]);
-							Tangent = (P3 - P2).GetSafeNormal();
-							RenderMesh->Tangents[3 * i + 1] = FProcMeshTangent(Tangent[0], Tangent[1], Tangent[2]);
-							Tangent = (P1 - P3).GetSafeNormal();
-							RenderMesh->Tangents[3 * i + 2] = FProcMeshTangent(Tangent[0], Tangent[1], Tangent[2]);
-						}
+								// Display only
+								for (int i = 0; i < NumFaces; ++i)
+								{
+									const auto& P1 = (*RenderVertex)[Flesh->Indices[i][0]];
+									const auto& P2 = (*RenderVertex)[Flesh->Indices[i][1]];
+									const auto& P3 = (*RenderVertex)[Flesh->Indices[i][2]];
+
+									RenderMesh->Vertices[3 * i] = FVector(P1);
+									RenderMesh->Vertices[3 * i + 1] = FVector(P2);
+									RenderMesh->Vertices[3 * i + 2] = FVector(P3);
+
+									auto Normal = Chaos::FVec3::CrossProduct(P3 - P1, P2 - P1);
+									RenderMesh->Normals[3 * i] = Normal;
+									RenderMesh->Normals[3 * i + 1] = Normal;
+									RenderMesh->Normals[3 * i + 2] = Normal;
+
+									auto Tangent = (P2 - P1).GetSafeNormal();
+									RenderMesh->Tangents[3 * i] = FProcMeshTangent(Tangent[0], Tangent[1], Tangent[2]);
+									Tangent = (P3 - P2).GetSafeNormal();
+									RenderMesh->Tangents[3 * i + 1] = FProcMeshTangent(Tangent[0], Tangent[1], Tangent[2]);
+									Tangent = (P1 - P3).GetSafeNormal();
+									RenderMesh->Tangents[3 * i + 2] = FProcMeshTangent(Tangent[0], Tangent[1], Tangent[2]);
+								}
 
-						if (!Mesh->GetComponentTransform().Equals(GetComponentTransform())) {
-							Mesh->SetRelativeTransform(GetComponentTransform());
+								if (!Mesh->GetComponentTransform().Equals(GetComponentTransform())) {
+									Mesh->SetRelativeTransform(GetComponentTransform());
+								}
+								Mesh->UpdateMeshSection_LinearColor(0, RenderMesh->Vertices, RenderMesh->Normals, RenderMesh->UVs, RenderMesh->Colors, RenderMesh->Tangents);
+							}
+
+							bCanRender = true;
 						}
-						Mesh->UpdateMeshSection_LinearColor(0, RenderMesh->Vertices, RenderMesh->Normals, RenderMesh->UVs, RenderMesh->Colors, RenderMesh->Tangents);
 					}
-
-					bCanRender = true;
 				}
 			}
 		}
diff --git a/Engine/Plugins/Experimental/ChaosFlesh/Source/ChaosFleshEngine/Public/ChaosFlesh/ChaosDeformableCollisionsActor.h b/Engine/Plugins/Experimental/ChaosFlesh/Source/ChaosFleshEngine/Public/ChaosFlesh/ChaosDeformableCollisionsActor.h
new file mode 100644
--- /dev/null
+++ b/Engine/Plugins/Experimental/ChaosFlesh/Source/ChaosFleshEngine/Public/ChaosFlesh/ChaosDeformableCollisionsActor.h
@@ -0,0 +1,45 @@
+// Copyright Epic Games, Inc. All Rights Reserved.
+
+#pragma once
+
+#include "CoreMinimal.h"
+#include "GameFramework/Actor.h"
+#include "UObject/ObjectMacros.h"
+#include "UObject/ObjectPtr.h"
+#include "ChaosFlesh/ChaosDeformableCollisionsComponent.h"
+#include "ChaosFlesh/ChaosDeformableSolverActor.h"
+
+#include "ChaosDeformableCollisionsActor.generated.h"
+
+class AStaticMeshActor;
+
+
+UCLASS()
+class CHAOSFLESHENGINE_API ADeformableCollisionsActor : public AActor
+{
+	GENERATED_UCLASS_BODY()
+
+public:
+
+	UFUNCTION(BlueprintCallable, Category = "Physics")
+	void EnableSimulation(ADeformableSolverActor* Actor);
+
+	UPROPERTY(VisibleAnywhere, BlueprintReadOnly, Category = "Chaos Deformable")
+	TObjectPtr<UDeformableCollisionsComponent> DeformableCollisionsComponent;
+	UDeformableCollisionsComponent* GetCollisionsComponent() const { return DeformableCollisionsComponent; }
+
+	UPROPERTY(EditAnywhere, BlueprintReadWrite, Category = "Chaos Deformable")
+	TObjectPtr<ADeformableSolverActor> PrimarySolver;
+
+	UPROPERTY(EditAnywhere, BlueprintReadWrite, Category = "Chaos Deformable")
+	TArray<TObjectPtr<AStaticMeshActor>> StaticCollisions;
+
+#if WITH_EDITOR
+	TArray < TObjectPtr<AStaticMeshActor> > AddedBodies;
+	TArray < TObjectPtr<AStaticMeshActor> > RemovedBodies;
+	TArray < TObjectPtr<AStaticMeshActor> > PreEditChangeCollisionBodies;
+	ADeformableSolverActor* PreEditChangePrimarySolver = nullptr;
+	virtual void PreEditChange(FProperty* PropertyThatWillChange) override;
+	virtual void PostEditChangeProperty(struct FPropertyChangedEvent& PropertyChangedEvent) override;
+#endif
+};
diff --git a/Engine/Plugins/Experimental/ChaosFlesh/Source/ChaosFleshEngine/Public/ChaosFlesh/ChaosDeformableCollisionsComponent.h b/Engine/Plugins/Experimental/ChaosFlesh/Source/ChaosFleshEngine/Public/ChaosFlesh/ChaosDeformableCollisionsComponent.h
new file mode 100644
--- /dev/null
+++ b/Engine/Plugins/Experimental/ChaosFlesh/Source/ChaosFleshEngine/Public/ChaosFlesh/ChaosDeformableCollisionsComponent.h
@@ -0,0 +1,53 @@
+// Copyright Epic Games, Inc. All Rights Reserved.
+#pragma once
+
+#include "Chaos/Deformable/ChaosDeformableSolverProxy.h"
+#include "Chaos/Deformable/ChaosDeformableCollisionsProxy.h"
+#include "ChaosFlesh/ChaosDeformablePhysicsComponent.h"
+#include "Components/SceneComponent.h"
+#include "UObject/ObjectMacros.h"
+
+#include "ChaosDeformableCollisionsComponent.generated.h"
+
+class UStaticMeshComponent;
+
+
+/**
+*	UDeformableCollisionsComponent
+*/
+UCLASS(meta = (BlueprintSpawnableComponent))
+class CHAOSFLESHENGINE_API UDeformableCollisionsComponent : public UDeformablePhysicsComponent
+{
+	GENERATED_UCLASS_BODY()
+
+public:
+	typedef Chaos::Softs::FDataMapValue FDataMapValue;
+	typedef Chaos::Softs::FThreadingProxy FThreadingProxy;
+	typedef Chaos::Softs::FCollisionManagerProxy FCollisionThreadingProxy;
+
+	~UDeformableCollisionsComponent() {}
+
+	/** Simulation Interface*/
+	virtual FThreadingProxy* NewProxy() override;
+	virtual FDataMapValue NewDeformableData() override;
+
+
+	UFUNCTION(BlueprintCallable, Category = "Physics")
+	void AddStaticMeshComponent(UStaticMeshComponent* StaticMeshComponent);
+	//void AddStaticMeshComponents(const TArray<TObjectPtr<UStaticMeshComponent>>& StaticMeshComponents);
+	
+	UFUNCTION(BlueprintCallable, Category = "Physics")
+	void RemoveStaticMeshComponent(UStaticMeshComponent* StaticMeshComponent);
+	//void RemoveStaticMeshComponents(const TArray<TObjectPtr<UStaticMeshComponent>>& StaticMeshComponents);
+
+	UPROPERTY(EditAnywhere, BlueprintReadWrite, Category = "Chaos Deformable")
+	TArray< TObjectPtr<UStaticMeshComponent> >  CollisionBodies;
+
+protected:
+	
+	TArray< UStaticMeshComponent* >  RemovedBodies;
+	TArray< UStaticMeshComponent* >  AddedBodies;
+
+	
+	TMap<TObjectPtr<UObject>, FThreadingProxy*> CollisionsMap;
+};
diff --git a/Engine/Plugins/Experimental/ChaosFlesh/Source/ChaosFleshEngine/Public/ChaosFlesh/ChaosDeformablePhysicsComponent.h b/Engine/Plugins/Experimental/ChaosFlesh/Source/ChaosFleshEngine/Public/ChaosFlesh/ChaosDeformablePhysicsComponent.h
--- a/Engine/Plugins/Experimental/ChaosFlesh/Source/ChaosFleshEngine/Public/ChaosFlesh/ChaosDeformablePhysicsComponent.h
+++ b/Engine/Plugins/Experimental/ChaosFlesh/Source/ChaosFleshEngine/Public/ChaosFlesh/ChaosDeformablePhysicsComponent.h
@@ -16,7 +16,7 @@ class UDeformableSolverComponent;
 *	UDeformablePhysicsComponent
 */
 UCLASS(meta = (BlueprintSpawnableComponent))
-class CHAOSFLESHENGINE_API UDeformablePhysicsComponent : public UMeshComponent
+class CHAOSFLESHENGINE_API UDeformablePhysicsComponent : public UPrimitiveComponent
 {
 	GENERATED_UCLASS_BODY()
 
diff --git a/Engine/Plugins/Experimental/ChaosFlesh/Source/ChaosFleshEngine/Public/ChaosFlesh/ChaosDeformableSolverComponent.h b/Engine/Plugins/Experimental/ChaosFlesh/Source/ChaosFleshEngine/Public/ChaosFlesh/ChaosDeformableSolverComponent.h
--- a/Engine/Plugins/Experimental/ChaosFlesh/Source/ChaosFleshEngine/Public/ChaosFlesh/ChaosDeformableSolverComponent.h
+++ b/Engine/Plugins/Experimental/ChaosFlesh/Source/ChaosFleshEngine/Public/ChaosFlesh/ChaosDeformableSolverComponent.h
@@ -13,6 +13,7 @@
 #include "ChaosDeformableSolverComponent.generated.h"
 
 class UDeformablePhysicsComponent;
+class UDeformableCollisionsComponent;
 
 /**
 *	UDeformableSolverComponent
@@ -55,6 +56,9 @@ public:
 	UPROPERTY(EditAnywhere, BlueprintReadOnly, Category = "Chaos Deformable")
 	TArray< TObjectPtr<UDeformablePhysicsComponent> > DeformableComponents;
 
+	UPROPERTY(EditAnywhere, BlueprintReadOnly, Category = "Chaos Deformable")
+	TObjectPtr<UDeformableCollisionsComponent> CollisionComponent;
+
 	/** ObjectType defines how to initialize the rigid objects state, Kinematic, Sleeping, Dynamic. */
 	UPROPERTY(EditAnywhere, BlueprintReadOnly , Category = "Chaos Deformable")
 	EDeformableExecutionModel ExecutionModel = EDeformableExecutionModel::Chaos_Deformable_DuringPhysics;
diff --git a/Engine/Source/Runtime/Experimental/Chaos/Private/Chaos/Deformable/ChaosDeformableSolver.cpp b/Engine/Source/Runtime/Experimental/Chaos/Private/Chaos/Deformable/ChaosDeformableSolver.cpp
--- a/Engine/Source/Runtime/Experimental/Chaos/Private/Chaos/Deformable/ChaosDeformableSolver.cpp
+++ b/Engine/Source/Runtime/Experimental/Chaos/Private/Chaos/Deformable/ChaosDeformableSolver.cpp
@@ -3,6 +3,7 @@
 #include "Chaos/Deformable/ChaosDeformableSolver.h"
 #include "Chaos/Deformable/ChaosDeformableSolverTypes.h"
 #include "Chaos/Deformable/ChaosDeformableSolverProxy.h"
+#include "Chaos/Deformable/ChaosDeformableCollisionsProxy.h"
 
 #include "Chaos/DebugDrawQueue.h"
 #include "Chaos/TriangleMesh.h"
@@ -139,6 +140,17 @@ namespace Chaos::Softs
 		Frame = 0;
 		Time = 0.f;
 		Iteration = 0;
+
+		// Add a default floor the first time through
+		if (Property.bUseFloor)
+		{
+			Chaos::FVec3 Position(0.f);
+			Chaos::FVec3 EulerRot(0.f);
+			int32 CollisionParticleOffset = Evolution->AddCollisionParticleRange(1, INDEX_NONE, true);
+			Evolution->CollisionParticles().X(0) = Position;
+			Evolution->CollisionParticles().R(0) = Chaos::TRotation<Chaos::FReal, 3>::MakeFromEuler(EulerRot);
+			Evolution->CollisionParticles().SetDynamicGeometry(0, MakeUnique<Chaos::TPlane<Chaos::FReal, 3>>(Chaos::FVec3(0.f, 0.f, 0.f), Chaos::FVec3(0.f, 0.f, 1.f)));
+		}
 	}
 
 	void FDeformableSolver::Simulate(FSolverReal DeltaTime)
@@ -186,8 +198,22 @@ namespace Chaos::Softs
 				UninitializedProxys_Internal.SetNum(0, true);
 			}
 		}
+	}
 
-		InitializeCollisionBodies();
+	void FDeformableSolver::UpdateSimulationObjects(FSolverReal DeltaTime)
+	{
+		SCOPE_CYCLE_COUNTER(STAT_ChaosDeformableSolver_InitializeSimulationObjects);
+		TRACE_CPUPROFILER_EVENT_SCOPE(ChaosDeformableSolver_InitializeSimulationObjects);
+
+		typedef TPair< FThreadingProxy::FKey, TUniquePtr<FThreadingProxy> > FType;
+		for (const FType& Entry : Proxies)
+		{
+			FThreadingProxy& Proxy = *Entry.Value.Get();
+			if (FCollisionManagerProxy* CollisionManagerProxy = Proxy.As< FCollisionManagerProxy>())
+			{
+				UpdateCollisionBodies(Entry.Key, DeltaTime);
+			}
+		}
 	}
 
 	void FDeformableSolver::InitializeSimulationObject(FThreadingProxy& InProxy)
@@ -202,6 +228,11 @@ namespace Chaos::Softs
 			InitializeWeakConstraint(*Proxy);
 			InitializeTetrahedralConstraint(*Proxy);
 			InitializeGidBasedConstraints(*Proxy);
+
+		}
+		if (FCollisionManagerProxy* CollisionManagerProxy = InProxy.As< FCollisionManagerProxy>())
+		{
+			InitializeCollisionBodies(*CollisionManagerProxy);
 		}
 	}
 
@@ -396,6 +427,42 @@ namespace Chaos::Softs
 	}
 
 
+	void FDeformableSolver::InitializeCollisionBodies(FCollisionManagerProxy& Proxy)
+	{
+		SCOPE_CYCLE_COUNTER(STAT_ChaosDeformableSolver_InitializeCollisionBodies);
+		TRACE_CPUPROFILER_EVENT_SCOPE(ChaosDeformableSolver_InitializeCollisionBodies);
+	}
+
+	void FDeformableSolver::UpdateCollisionBodies(FThreadingProxy::FKey Owner, FSolverReal DeltaTime)
+	{
+		SCOPE_CYCLE_COUNTER(STAT_ChaosDeformableSolver_InitializeCollisionBodies);
+		TRACE_CPUPROFILER_EVENT_SCOPE(ChaosDeformableSolver_InitializeCollisionBodies);
+
+		FCollisionManagerProxy::FCollisionsInputBuffer* CollisionsInputBuffer = nullptr;
+		if (this->CurrentInputPackage->ObjectMap.Contains(Owner))
+		{
+			CollisionsInputBuffer = this->CurrentInputPackage->ObjectMap[Owner]->As<FCollisionManagerProxy::FCollisionsInputBuffer>();
+			if (CollisionsInputBuffer)
+			{
+				for (auto& AddBody : CollisionsInputBuffer->Added)
+				{
+					if (AddBody.Shapes)
+					{
+						typedef TPlane<Chaos::Softs::FSolverReal, 3> GeomType;
+						int32 Index = Evolution->AddCollisionParticleRange(1, INDEX_NONE, true);
+						Evolution->CollisionParticles().X(Index) = AddBody.Transform.GetTranslation();
+						Evolution->CollisionParticles().R(Index) = AddBody.Transform.GetRotation();
+						TUniquePtr<FImplicitObject> UniquePtr(AddBody.Shapes); AddBody.Shapes = nullptr;
+						Evolution->CollisionParticles().SetDynamicGeometry(Index, MoveTemp(UniquePtr));
+					}
+				}
+			}
+		}
+
+	}
+
+
+
 	void FDeformableSolver::DebugDrawTetrahedralParticles(FFleshThreadingProxy& Proxy)
 	{
 #if WITH_EDITOR
@@ -701,23 +768,6 @@ namespace Chaos::Softs
 		Evolution->SetKinematicUpdateFunction(MKineticUpdate);
 	}
 
-	void FDeformableSolver::InitializeCollisionBodies()
-	{
-		SCOPE_CYCLE_COUNTER(STAT_ChaosDeformableSolver_InitializeCollisionBodies);
-		TRACE_CPUPROFILER_EVENT_SCOPE(ChaosDeformableSolver_InitializeCollisionBodies);
-
-		if (Property.bUseFloor && Evolution->CollisionParticles().Size() == 0)
-		{
-			Chaos::FVec3 Position(0.f);
-			Chaos::FVec3 EulerRot(0.f);
-			int32 CollisionParticleOffset = Evolution->AddCollisionParticleRange(1, INDEX_NONE, true);
-			Evolution->CollisionParticles().X(0) = Position;
-			Evolution->CollisionParticles().R(0) = Chaos::TRotation<Chaos::FReal, 3>::MakeFromEuler(EulerRot);
-			Evolution->CollisionParticles().SetDynamicGeometry(0, MakeUnique<Chaos::TPlane<Chaos::FReal, 3>>(Chaos::FVec3(0.f, 0.f, 0.f), Chaos::FVec3(0.f, 0.f, 1.f)));
-
-		}
-	}
-
 	void FDeformableSolver::InitializeSelfCollisionVariables()
 	{
 		SCOPE_CYCLE_COUNTER(STAT_ChaosDeformableSolver_InitializeSelfCollisionVariables);
@@ -964,6 +1014,7 @@ namespace Chaos::Softs
 
 		if (!Proxies.Num()) return;
 
+		UpdateSimulationObjects(DeltaTime);
 
 		if (!Property.FixTimeStep)
 		{
@@ -1125,4 +1176,5 @@ namespace Chaos::Softs
 		FFileHelper::SaveStringToFile(FString(TEXT("endExtra\n")), *file, FFileHelper::EEncodingOptions::AutoDetect, &IFileManager::Get(), EFileWrite::FILEWRITE_Append);
 
 	}
+
 }; // Namespace Chaos::Softs
\ No newline at end of file
diff --git a/Engine/Source/Runtime/Experimental/Chaos/Private/Chaos/Deformable/ChaosDeformableSolverThreadingWrapper.cpp b/Engine/Source/Runtime/Experimental/Chaos/Private/Chaos/Deformable/ChaosDeformableSolverThreadingWrapper.cpp
--- a/Engine/Source/Runtime/Experimental/Chaos/Private/Chaos/Deformable/ChaosDeformableSolverThreadingWrapper.cpp
+++ b/Engine/Source/Runtime/Experimental/Chaos/Private/Chaos/Deformable/ChaosDeformableSolverThreadingWrapper.cpp
@@ -63,10 +63,10 @@ namespace Chaos::Softs
 			Solver->InitializeSimulationObject(InProxy);
 	}
 
-	void FDeformableSolver::FPhysicsThreadAccess::InitializeCollisionBodies()
+	void FDeformableSolver::FPhysicsThreadAccess::InitializeCollisionBodies(FCollisionManagerProxy& Proxy)
 	{
 		if (Solver)
-			Solver->InitializeCollisionBodies();
+			Solver->InitializeCollisionBodies(Proxy);
 	}
 
 	void FDeformableSolver::FPhysicsThreadAccess::InitializeKinematicConstraint()
diff --git a/Engine/Source/Runtime/Experimental/Chaos/Public/Chaos/Deformable/ChaosDeformableCollisionsProxy.h b/Engine/Source/Runtime/Experimental/Chaos/Public/Chaos/Deformable/ChaosDeformableCollisionsProxy.h
new file mode 100644
--- /dev/null
+++ b/Engine/Source/Runtime/Experimental/Chaos/Public/Chaos/Deformable/ChaosDeformableCollisionsProxy.h
@@ -0,0 +1,79 @@
+// Copyright Epic Games, Inc. All Rights Reserved.
+
+#pragma once
+#include "Chaos/Deformable/ChaosDeformableSolverProxy.h"
+#include "CoreMinimal.h"
+#include "UObject/ObjectMacros.h"
+
+namespace Chaos::Softs
+{
+	struct CHAOS_API FCollisionObjectAddedBodies
+	{
+		FCollisionObjectAddedBodies( const UObject* InBodyId = nullptr,
+			FTransform InTransform = FTransform::Identity,
+			FString InType = "",
+			FImplicitObject* InShapes = nullptr)
+			: BodyId(InBodyId)
+			, Transform(InTransform)
+			, Type(InType)
+			, Shapes(InShapes) {}
+
+	
+		const UObject* BodyId = nullptr;
+		FTransform Transform = FTransform::Identity;
+		FString Type = "";
+		FImplicitObject* Shapes = nullptr;
+	};
+
+	struct CHAOS_API FCollisionObjectRemovedBodies
+	{
+		const UObject* BodyId = nullptr;
+	};
+
+	struct CHAOS_API FCollisionObjectUpdatedBodies
+	{
+		const UObject* BodyId = nullptr;
+		FTransform Transform = FTransform::Identity;
+	};
+
+
+	class CHAOS_API FCollisionManagerProxy : public FThreadingProxy
+	{
+	public:
+		typedef FThreadingProxy Super;
+
+		FCollisionManagerProxy(UObject* InOwner)
+			: Super(InOwner, TypeName())
+		{}
+
+		static FName TypeName() { return FName("CollisionManager"); }
+
+		class FCollisionsInputBuffer : public FThreadingProxy::FBuffer
+		{
+			typedef FThreadingProxy::FBuffer Super;
+
+		public:
+			typedef FCollisionManagerProxy Source;
+
+			FCollisionsInputBuffer(
+				const TArray<FCollisionObjectAddedBodies>& InAdded
+				, const TArray<FCollisionObjectRemovedBodies>& InRemoved
+				, const TArray<FCollisionObjectUpdatedBodies>& InUpdate
+				, const UObject* InOwner)
+				: Super(InOwner, FCollisionManagerProxy::TypeName())
+				, Added(InAdded)
+				, Removed(InRemoved)
+				, Updated(InUpdate)
+			{}
+
+			TArray<FCollisionObjectAddedBodies> Added;
+			TArray<FCollisionObjectRemovedBodies> Removed;
+			TArray<FCollisionObjectUpdatedBodies> Updated;
+		};
+
+		TArray<FCollisionObjectAddedBodies> CollisionObjectsToAdd;
+		TArray< FCollisionObjectRemovedBodies > CollisionObjectsToRemove;
+		TMap< UObject*, FTransform > CollisionBodies;
+	};
+
+}// namespace Chaos::Softs
\ No newline at end of file
diff --git a/Engine/Source/Runtime/Experimental/Chaos/Public/Chaos/Deformable/ChaosDeformableSolver.h b/Engine/Source/Runtime/Experimental/Chaos/Public/Chaos/Deformable/ChaosDeformableSolver.h
--- a/Engine/Source/Runtime/Experimental/Chaos/Public/Chaos/Deformable/ChaosDeformableSolver.h
+++ b/Engine/Source/Runtime/Experimental/Chaos/Public/Chaos/Deformable/ChaosDeformableSolver.h
@@ -16,6 +16,7 @@
 #include "Chaos/XPBDGridBasedCorotatedConstraints.h"
 #include "Chaos/Deformable/ChaosDeformableSolverTypes.h"
 #include "Chaos/Deformable/ChaosDeformableSolverProxy.h"
+#include "Chaos/Deformable/ChaosDeformableCollisionsProxy.h"
 #include "CoreMinimal.h"
 #include "GeometryCollection/ManagedArrayCollection.h"
 #include "Templates/UniquePtr.h"
@@ -52,7 +53,7 @@ namespace Chaos::Softs
 			/* Iteration Advance */
 			void InitializeSimulationObjects();
 			void InitializeSimulationObject(FThreadingProxy&);
-			void InitializeCollisionBodies();
+			void InitializeCollisionBodies(FCollisionManagerProxy&);
 			void InitializeKinematicConstraint();
 			void InitializeSelfCollisionVariables();
 			void RemoveSimulationObjects();
@@ -109,6 +110,7 @@ namespace Chaos::Softs
 		void AdvanceDt(FSolverReal DeltaTime);
 		void Reset(const FDeformableSolverProperties&);
 		void Update(FSolverReal DeltaTime);
+		void UpdateSimulationObjects(FSolverReal DeltaTime);
 		void UpdateOutputState(FThreadingProxy&);
 		void PushOutputPackage(int32 Frame, FDeformableDataMap&& Package);
 		TUniquePtr<FDeformablePackage> PullInputPackage( );
@@ -122,9 +124,10 @@ namespace Chaos::Softs
 		void InitializeGidBasedConstraints(FFleshThreadingProxy&);
 		void InitializeWeakConstraint(FFleshThreadingProxy&);
 		void InitializeKinematicConstraint();
-		void InitializeCollisionBodies();
+		void InitializeCollisionBodies(FCollisionManagerProxy&);
 		void InitializeSelfCollisionVariables();
 		void InitializeGridBasedConstraintVariables();
+		void UpdateCollisionBodies(FThreadingProxy::FKey, FSolverReal DeltaTime);
 		void RemoveSimulationObjects();
 		TArray<Chaos::TVec3<FSolverReal>> ComputeParticleTargets(const TArray<TArray<int32>>& ParticleIndices);
 
diff --git a/Engine/Source/Runtime/Experimental/Chaos/Public/Chaos/Deformable/ChaosDeformableSolverTypes.h b/Engine/Source/Runtime/Experimental/Chaos/Public/Chaos/Deformable/ChaosDeformableSolverTypes.h
--- a/Engine/Source/Runtime/Experimental/Chaos/Public/Chaos/Deformable/ChaosDeformableSolverTypes.h
+++ b/Engine/Source/Runtime/Experimental/Chaos/Public/Chaos/Deformable/ChaosDeformableSolverTypes.h
@@ -4,6 +4,7 @@
 
 #include "Chaos/DebugDrawQueue.h"
 #include "Chaos/Deformable/ChaosDeformableSolverProxy.h"
+#include "Chaos/Deformable/ChaosDeformableCollisionsProxy.h"
 #include "Chaos/PBDSoftsEvolutionFwd.h"
 #include "CoreMinimal.h"
 
