ChaosFlesh : Expose Poisson ratio and J term coefficient for XPBD corotated constraints.

diff --git a/Engine/Plugins/Experimental/ChaosFlesh/Source/ChaosFleshEngine/Private/ChaosFlesh/FleshComponent.cpp b/Engine/Plugins/Experimental/ChaosFlesh/Source/ChaosFleshEngine/Private/ChaosFlesh/FleshComponent.cpp
--- a/Engine/Plugins/Experimental/ChaosFlesh/Source/ChaosFleshEngine/Private/ChaosFlesh/FleshComponent.cpp
+++ b/Engine/Plugins/Experimental/ChaosFlesh/Source/ChaosFleshEngine/Private/ChaosFlesh/FleshComponent.cpp
@@ -236,6 +236,8 @@ UDeformablePhysicsComponent::FDataMapValue UFleshComponent::NewDeformableData()
 								StiffnessMultiplier, 
 								DampingMultiplier, 
 								MassMultiplier, 
+								IncompressibilityMultiplier,
+								InflationMultiplier,
 								this));
 					}
 				}
@@ -251,6 +253,8 @@ UDeformablePhysicsComponent::FDataMapValue UFleshComponent::NewDeformableData()
 			StiffnessMultiplier, 
 			DampingMultiplier, 
 			MassMultiplier, 
+			IncompressibilityMultiplier,
+			InflationMultiplier,
 			this));
 }
 
diff --git a/Engine/Plugins/Experimental/ChaosFlesh/Source/ChaosFleshEngine/Public/ChaosFlesh/ChaosDeformablePhysicsComponent.h b/Engine/Plugins/Experimental/ChaosFlesh/Source/ChaosFleshEngine/Public/ChaosFlesh/ChaosDeformablePhysicsComponent.h
--- a/Engine/Plugins/Experimental/ChaosFlesh/Source/ChaosFleshEngine/Public/ChaosFlesh/ChaosDeformablePhysicsComponent.h
+++ b/Engine/Plugins/Experimental/ChaosFlesh/Source/ChaosFleshEngine/Public/ChaosFlesh/ChaosDeformablePhysicsComponent.h
@@ -64,6 +64,12 @@ public:
 	UPROPERTY(EditAnywhere, BlueprintReadWrite, Category = "ChaosDeformable", meta = (ClampMin = "0.0", ClampMax = "1.0", UIMin = "0.0", UIMax = "1.0"))
 	float MassMultiplier = 1.f;
 
+	UPROPERTY(EditAnywhere, BlueprintReadWrite, Category = "ChaosDeformable", meta = (ClampMin = "0.0", ClampMax = "1.0", UIMin = "0.0", UIMax = "1.0"))
+	float IncompressibilityMultiplier = 1.f;
+
+	UPROPERTY(EditAnywhere, BlueprintReadWrite, Category = "ChaosDeformable", meta = (ClampMin = "0.0", ClampMax = "1.0", UIMin = "0.0", UIMax = "1.0"))
+	float InflationMultiplier = 1.f;
+
 	const FThreadingProxy* GetPhysicsProxy() const { return PhysicsProxy; }
 	FThreadingProxy* GetPhysicsProxy() { return PhysicsProxy; }
 
diff --git a/Engine/Plugins/Experimental/ChaosFlesh/Source/ChaosFleshNodes/Private/Dataflow/ChaosFleshEngineAssetNodes.cpp b/Engine/Plugins/Experimental/ChaosFlesh/Source/ChaosFleshNodes/Private/Dataflow/ChaosFleshEngineAssetNodes.cpp
--- a/Engine/Plugins/Experimental/ChaosFlesh/Source/ChaosFleshNodes/Private/Dataflow/ChaosFleshEngineAssetNodes.cpp
+++ b/Engine/Plugins/Experimental/ChaosFlesh/Source/ChaosFleshNodes/Private/Dataflow/ChaosFleshEngineAssetNodes.cpp
@@ -66,20 +66,19 @@ void FSetFleshDefaultPropertiesNode::Evaluate(Dataflow::FContext& Context, const
 	{
 		FManagedArrayCollection InCollection = GetValue<FManagedArrayCollection>(Context, &Collection);
 		
-		//if (InCollection.HasAttributes({
-		//	MType< float >("Stiffness", FGeometryCollection::VerticesGroup),
-		//	MType< float >("Damping", FGeometryCollection::VerticesGroup) })
-		//{
-
-			TManagedArray<float>& ParticleStiffness = InCollection.AddAttribute<float>("Stiffness", FGeometryCollection::VerticesGroup);
-				TManagedArray<float>& ParticleDamping = InCollection.AddAttribute<float>("Damping", FGeometryCollection::VerticesGroup);
-		//}//TManagedArray<float>& Stiffness = InCollection.ModifyAttribute<float>("Stiffness", FGeometryCollection::VerticesGroup);
+		TManagedArray<float>& ParticleStiffness = InCollection.AddAttribute<float>("Stiffness", FGeometryCollection::VerticesGroup);
+		TManagedArray<float>& ParticleDamping = InCollection.AddAttribute<float>("Damping", FGeometryCollection::VerticesGroup);
+		TManagedArray<float>& ParticleIncompressibility = InCollection.AddAttribute<float>("Incompressibility", FGeometryCollection::VerticesGroup);
+		TManagedArray<float>& ParticleInflation = InCollection.AddAttribute<float>("Inflation", FGeometryCollection::VerticesGroup);
+
 		
 
 		if (InCollection.HasAttributes({
 			MType< float >("Mass", FGeometryCollection::VerticesGroup),
 			MType< float >("Stiffness", FGeometryCollection::VerticesGroup),
 			MType< float >("Damping", FGeometryCollection::VerticesGroup),
+			MType< float >("Incompressibility", FGeometryCollection::VerticesGroup),
+			MType< float >("Inflation", FGeometryCollection::VerticesGroup),
 			MType< FIntVector4 >(FTetrahedralCollection::TetrahedronAttribute, FTetrahedralCollection::TetrahedralGroup),
 			MType< FVector3f >("Vertex", "Vertices"),
 			MType< TArray<int32> >(FTetrahedralCollection::IncidentElementsAttribute, FGeometryCollection::VerticesGroup),
@@ -93,6 +92,8 @@ void FSetFleshDefaultPropertiesNode::Evaluate(Dataflow::FContext& Context, const
 				TManagedArray<float>& Mass = InCollection.ModifyAttribute<float>("Mass", FGeometryCollection::VerticesGroup);
 				TManagedArray<float>& Stiffness = InCollection.ModifyAttribute<float>("Stiffness", FGeometryCollection::VerticesGroup);
 				TManagedArray<float>& Damping = InCollection.ModifyAttribute<float>("Damping", FGeometryCollection::VerticesGroup);
+				TManagedArray<float>& Incompressibility = InCollection.ModifyAttribute<float>("Incompressibility", FGeometryCollection::VerticesGroup);
+				TManagedArray<float>& Inflation = InCollection.ModifyAttribute<float>("Inflation", FGeometryCollection::VerticesGroup);
 				const TManagedArray<FIntVector4>& Tetrahedron = InCollection.GetAttribute<FIntVector4>(FTetrahedralCollection::TetrahedronAttribute, FTetrahedralCollection::TetrahedralGroup);
 				const TManagedArray<FVector3f>& Vertex = InCollection.GetAttribute<FVector3f>("Vertex", "Vertices");
 				const TManagedArray<TArray<int32>>& IncidentElements = InCollection.GetAttribute<TArray<int32>>(FTetrahedralCollection::IncidentElementsAttribute, FGeometryCollection::VerticesGroup);
@@ -181,6 +182,8 @@ void FSetFleshDefaultPropertiesNode::Evaluate(Dataflow::FContext& Context, const
 
 				Stiffness.Fill(VertexStiffness);
 				Damping.Fill(VertexDamping);
+				Incompressibility.Fill(.5f * VertexIncompressibility);
+				Inflation.Fill(VertexInflation * 2.f);
 
 				UE_LOG(LogChaosFlesh, Display,
 					TEXT("'%s' - Set mass on %d nodes:\n"
diff --git a/Engine/Plugins/Experimental/ChaosFlesh/Source/ChaosFleshNodes/Public/Dataflow/ChaosFleshEngineAssetNodes.h b/Engine/Plugins/Experimental/ChaosFlesh/Source/ChaosFleshNodes/Public/Dataflow/ChaosFleshEngineAssetNodes.h
--- a/Engine/Plugins/Experimental/ChaosFlesh/Source/ChaosFleshNodes/Public/Dataflow/ChaosFleshEngineAssetNodes.h
+++ b/Engine/Plugins/Experimental/ChaosFlesh/Source/ChaosFleshNodes/Public/Dataflow/ChaosFleshEngineAssetNodes.h
@@ -77,6 +77,16 @@ public:
 	UPROPERTY(EditAnywhere, Category = "Dataflow", meta = (ClampMin = "0.0", ClampMax = "1.0", UIMin = "0.0", UIMax = "1.0"))
 	float VertexDamping = 0.f;
 
+	/*Sets incompressibility on vertex basis. 0.6 is default behavior. 
+	1 means total incompressibility. 0.00001 means almost no incompressibility*/
+	UPROPERTY(EditAnywhere, Category = "Dataflow", meta = (ClampMin = "0.00001", ClampMax = "1.0", UIMin = "0.00001", UIMax = "1.0"))
+	float VertexIncompressibility = 0.6f;
+
+	/*Sets inflation on vertex basis. 0.5 means no inflation/deflation.
+	1 means total inflation. 0 means the material is deflated.*/
+	UPROPERTY(EditAnywhere, Category = "Dataflow", meta = (ClampMin = "0.0", ClampMax = "1.0", UIMin = "0.0", UIMax = "1.0"))
+	float VertexInflation = 0.5f;
+
 	UPROPERTY(meta = (DataflowInput, DataflowOutput, DisplayName = "Collection", DataflowPassthrough  = "Collection"))
 	FManagedArrayCollection Collection;
 
diff --git a/Engine/Source/Runtime/Experimental/Chaos/Private/Chaos/Deformable/ChaosDeformableSolver.cpp b/Engine/Source/Runtime/Experimental/Chaos/Private/Chaos/Deformable/ChaosDeformableSolver.cpp
--- a/Engine/Source/Runtime/Experimental/Chaos/Private/Chaos/Deformable/ChaosDeformableSolver.cpp
+++ b/Engine/Source/Runtime/Experimental/Chaos/Private/Chaos/Deformable/ChaosDeformableSolver.cpp
@@ -592,7 +592,9 @@ namespace Chaos::Softs
 				uint32 NumParticles = Rest.NumElements(FGeometryCollection::VerticesGroup);
 				TArray<FSolverReal> StiffnessWithMultiplier;
 				StiffnessWithMultiplier.Init(0.f, NumParticles);
-				FSolverReal StiffnessMultiplier = 0.f;
+				FSolverReal StiffnessMultiplier = 1.f;
+				FSolverReal IncompressibilityMultiplier = 1.f;
+				FSolverReal InflationMultiplier = 1.f;
 
 				if (const UObject* Owner = this->MObjects[Range[0]]) {
 					FFleshThreadingProxy::FFleshInputBuffer* FleshInputBuffer = nullptr;
@@ -602,6 +604,8 @@ namespace Chaos::Softs
 						if (FleshInputBuffer)
 						{
 							StiffnessMultiplier = FleshInputBuffer->StiffnessMultiplier;
+							IncompressibilityMultiplier = FleshInputBuffer->IncompressibilityMultiplier;
+							InflationMultiplier = FleshInputBuffer->InflationMultiplier;
 						}
 					}
 				}
@@ -619,7 +623,40 @@ namespace Chaos::Softs
 						TetStiffness[edx] = (StiffnessWithMultiplier[Tetrahedron[edx].X] + StiffnessWithMultiplier[Tetrahedron[edx].Y]
 							+ StiffnessWithMultiplier[Tetrahedron[edx].Z] + StiffnessWithMultiplier[Tetrahedron[edx].W]) / 4.f;
 					}
+				}
+
+				const TManagedArray<FSolverReal>* IncompressibilityArray = Rest.FindAttribute<FSolverReal>("Incompressibility", FGeometryCollection::VerticesGroup);
+				TArray<FSolverReal> TetNu, AlphaJMesh, IncompressibilityWithMultiplier, InflationWithMultiplier;
+				TetNu.Init(.3f, Elements.Num());
+
+				IncompressibilityWithMultiplier.Init(0.f, NumParticles);
+				InflationWithMultiplier.Init(0.f, NumParticles);
+				if (IncompressibilityArray)
+				{
+					for (uint32 vdx = 0; vdx < NumParticles; ++vdx)
+					{
+						IncompressibilityWithMultiplier[vdx] = (*IncompressibilityArray)[vdx] * IncompressibilityMultiplier;
+					}
+					for (int32 edx = 0; edx < Elements.Num(); edx++)
+					{
+						TetNu[edx] = (IncompressibilityWithMultiplier[Tetrahedron[edx].X] + IncompressibilityWithMultiplier[Tetrahedron[edx].Y]
+							+ IncompressibilityWithMultiplier[Tetrahedron[edx].Z] + IncompressibilityWithMultiplier[Tetrahedron[edx].W]) / 4.f;
+					}
+				}
 
+				const TManagedArray<FSolverReal>* InflationArray = Rest.FindAttribute<FSolverReal>("Inflation", FGeometryCollection::VerticesGroup);
+				AlphaJMesh.Init(1.f, Elements.Num());
+				if (InflationArray)
+				{
+					for (uint32 vdx = 0; vdx < NumParticles; ++vdx)
+					{
+						InflationWithMultiplier[vdx] = (*InflationArray)[vdx] * InflationMultiplier;
+					}
+					for (int32 edx = 0; edx < Elements.Num(); edx++)
+					{
+						AlphaJMesh[edx] = (InflationWithMultiplier[Tetrahedron[edx].X] + InflationWithMultiplier[Tetrahedron[edx].Y]
+							+ InflationWithMultiplier[Tetrahedron[edx].Z] + InflationWithMultiplier[Tetrahedron[edx].W]) / 4.f;
+					}
 				}
 
 				if (Property.bEnableCorotatedConstraints)
@@ -653,7 +690,7 @@ namespace Chaos::Softs
 					{
 						FXPBDCorotatedConstraints<FSolverReal, FSolverParticles>* CorotatedConstraint =
 							new FXPBDCorotatedConstraints<FSolverReal, FSolverParticles>(
-								Evolution->Particles(), Elements, TetStiffness, GDeformableXPBDCorotatedParams);
+								Evolution->Particles(), Elements, TetStiffness, TetNu, MoveTemp(AlphaJMesh), GDeformableXPBDCorotatedParams);
 
 						Evolution->ConstraintInits()[InitIndex] =
 							[CorotatedConstraint](FSolverParticles& InParticles, const FSolverReal Dt)
diff --git a/Engine/Source/Runtime/Experimental/Chaos/Public/Chaos/Deformable/ChaosDeformableSolverProxy.h b/Engine/Source/Runtime/Experimental/Chaos/Public/Chaos/Deformable/ChaosDeformableSolverProxy.h
--- a/Engine/Source/Runtime/Experimental/Chaos/Public/Chaos/Deformable/ChaosDeformableSolverProxy.h
+++ b/Engine/Source/Runtime/Experimental/Chaos/Public/Chaos/Deformable/ChaosDeformableSolverProxy.h
@@ -253,7 +253,9 @@ namespace Chaos::Softs
 				const bool InbEnableGravity,
 				const float InStiffnessMultiplier, 
 				const float InDampingMultiplier, 
-				const float InMassMultiplier, 
+				const float InMassMultiplier,
+				const float InInCompressibilityMultiplier,
+				const float InInflationMultiplier,
 				const UObject* InOwner = nullptr)
 				: Super(InOwner, FFleshThreadingProxy::TypeName())
 				, WorldToComponentXf(InWorldToComponentXf)
@@ -264,6 +266,8 @@ namespace Chaos::Softs
 				, StiffnessMultiplier(InStiffnessMultiplier)
 				, DampingMultiplier(InDampingMultiplier)
 				, MassMultiplier(InMassMultiplier)
+				, IncompressibilityMultiplier(InInCompressibilityMultiplier)
+				, InflationMultiplier(InInflationMultiplier)
 				, SimSpaceBoneIndex(InSimSpaceBoneIndex)
 			{}
 
@@ -277,6 +281,8 @@ namespace Chaos::Softs
 				const float InStiffnessMultiplier, 
 				const float InDampingMultiplier, 
 				const float InMassMultiplier, 
+				const float InInCompressibilityMultiplier,
+				const float InInflationMultiplier,
 				const UObject* InOwner = nullptr)
 				: Super(InOwner, FFleshThreadingProxy::TypeName())
 				, WorldToComponentXf(InWorldToComponentXf)
@@ -287,6 +293,8 @@ namespace Chaos::Softs
 				, StiffnessMultiplier(InStiffnessMultiplier)
 				, DampingMultiplier(InDampingMultiplier)
 				, MassMultiplier(InMassMultiplier)
+				, IncompressibilityMultiplier(InInCompressibilityMultiplier)
+				, InflationMultiplier(InInflationMultiplier)
 				, SimSpaceBoneIndex(InSimSpaceBoneIndex)
 			{}
 			virtual ~FFleshInputBuffer() {}
@@ -300,6 +308,8 @@ namespace Chaos::Softs
 			float StiffnessMultiplier;
 			float DampingMultiplier;
 			float MassMultiplier;
+			float IncompressibilityMultiplier;
+			float InflationMultiplier;
 			int32 SimSpaceBoneIndex = INDEX_NONE;
 		};
 
diff --git a/Engine/Source/Runtime/Experimental/Chaos/Public/Chaos/XPBDCorotatedConstraints.h b/Engine/Source/Runtime/Experimental/Chaos/Public/Chaos/XPBDCorotatedConstraints.h
--- a/Engine/Source/Runtime/Experimental/Chaos/Public/Chaos/XPBDCorotatedConstraints.h
+++ b/Engine/Source/Runtime/Experimental/Chaos/Public/Chaos/XPBDCorotatedConstraints.h
@@ -113,11 +113,13 @@ namespace Chaos::Softs
 			const ParticleType& InParticles,
 			const TArray<TVector<int32, 4>>& InMesh,
 			const TArray<T>& EMeshArray,
+			const TArray<T>& NuMeshArray,
+			TArray<T>&& AlphaJMeshArray,
 			const FDeformableXPBDCorotatedParams& InParams,
 			const T& NuMesh = (T).3,
 			const bool bRecordMetricIn = false
 		)
-			: CorotatedParams(InParams), bRecordMetric(bRecordMetricIn), MeshConstraints(InMesh)
+			: CorotatedParams(InParams), AlphaJArray(MoveTemp(AlphaJMeshArray)), bRecordMetric(bRecordMetricIn), MeshConstraints(InMesh)
 		{
 			ensureMsgf(EMeshArray.Num() == InMesh.Num(), TEXT("Input Young Modulus Array Size is wrong"));
 			LambdaArray.Init((T)0., 2 * MeshConstraints.Num());
@@ -136,8 +138,8 @@ namespace Chaos::Softs
 
 			for (int e = 0; e < InMesh.Num(); e++)
 			{
-				LambdaElementArray[e] = EMeshArray[e] * NuMesh / (((T)1. + NuMesh) * ((T)1. - (T)2. * NuMesh));
-				MuElementArray[e] = EMeshArray[e] / ((T)2. * ((T)1. + NuMesh));
+				LambdaElementArray[e] = EMeshArray[e] * NuMeshArray[e] / (((T)1. + NuMeshArray[e]) * ((T)1. - (T)2. * NuMeshArray[e]));
+				MuElementArray[e] = EMeshArray[e] / ((T)2. * ((T)1. + NuMeshArray[e]));
 
 				PMatrix<T, 3, 3> Dm = DsInit(e, InParticles);
 				PMatrix<T, 3, 3> DmInv = Dm.Inverse();
@@ -447,7 +449,7 @@ namespace Chaos::Softs
 			PMatrix<T, 3, 3> DmInvT = ElementDmInv(ElementIndex).GetTransposed();
 			
 			T J = Fe.Determinant();
-			if (J - 1 < Tol )
+			if (J - AlphaJArray[ElementIndex] < Tol)
 			{
 				return TVec4<TVector<T, 3>>(TVector<T, 3>((T)0.));
 			}
@@ -471,7 +473,7 @@ namespace Chaos::Softs
 			}
 			
 
-			T DLambda = (1 - J) - AlphaTilde * LambdaArray[2 * ElementIndex + 1];
+			T DLambda = (AlphaJArray[ElementIndex] - J) - AlphaTilde * LambdaArray[2 * ElementIndex + 1];
 
 			T Denom = AlphaTilde;
 			for (int i = 0; i < 4; i++)
@@ -579,6 +581,7 @@ namespace Chaos::Softs
 		T Lambda;
 		TArray<T> MuElementArray;
 		TArray<T> LambdaElementArray;
+		TArray<T> AlphaJArray;
 		mutable T HError;
 		mutable TArray<T> HErrorArray;
 		bool bRecordMetric;
