ChaosFlesh : Removed access to the solver and physics components from the details panel.

diff --git a/Engine/Plugins/Experimental/ChaosFlesh/Source/ChaosFleshEngine/Public/ChaosFlesh/ChaosDeformableCollisionsComponent.h b/Engine/Plugins/Experimental/ChaosFlesh/Source/ChaosFleshEngine/Public/ChaosFlesh/ChaosDeformableCollisionsComponent.h
--- a/Engine/Plugins/Experimental/ChaosFlesh/Source/ChaosFleshEngine/Public/ChaosFlesh/ChaosDeformableCollisionsComponent.h
+++ b/Engine/Plugins/Experimental/ChaosFlesh/Source/ChaosFleshEngine/Public/ChaosFlesh/ChaosDeformableCollisionsComponent.h
@@ -40,7 +40,7 @@ public:
 	void RemoveStaticMeshComponent(UStaticMeshComponent* StaticMeshComponent);
 	//void RemoveStaticMeshComponents(const TArray<TObjectPtr<UStaticMeshComponent>>& StaticMeshComponents);
 
-	UPROPERTY(EditAnywhere, BlueprintReadWrite, Category = "Chaos Deformable")
+	UPROPERTY(BlueprintReadOnly, Category = "Chaos Deformable")
 	TArray< TObjectPtr<UStaticMeshComponent> >  CollisionBodies;
 
 protected:
diff --git a/Engine/Plugins/Experimental/ChaosFlesh/Source/ChaosFleshEngine/Public/ChaosFlesh/FleshComponent.h b/Engine/Plugins/Experimental/ChaosFlesh/Source/ChaosFleshEngine/Public/ChaosFlesh/FleshComponent.h
--- a/Engine/Plugins/Experimental/ChaosFlesh/Source/ChaosFleshEngine/Public/ChaosFlesh/FleshComponent.h
+++ b/Engine/Plugins/Experimental/ChaosFlesh/Source/ChaosFleshEngine/Public/ChaosFlesh/FleshComponent.h
@@ -58,7 +58,7 @@ public:
 	const UFleshDynamicAsset* GetDynamicCollection() const { return DynamicCollection; }
 
 
-	UPROPERTY(EditAnywhere, BlueprintReadWrite, Category = Attributes)
+	UPROPERTY()
 	TObjectPtr<UProceduralMeshComponent> Mesh;
 
 	UPROPERTY(EditAnywhere, BlueprintReadWrite, Category = "Rendering")
@@ -71,7 +71,7 @@ private:
 	TObjectPtr<const UFleshAsset> RestCollection;
 
 	// Current simulation state.
-	UPROPERTY(EditAnywhere, Category = "ChaosPhysics")
+	UPROPERTY()
 	TObjectPtr<UFleshDynamicAsset> DynamicCollection;
 
 	//
