Added UFUNCTION GetSkeletalMeshEmbeddedPositions, which returns either positions or deltas of the bones from the TargetDeformationSkeleton, in world, component, or bone space.

diff --git a/Engine/Plugins/Experimental/ChaosFlesh/Source/ChaosFleshEngine/Private/ChaosFlesh/FleshComponent.cpp b/Engine/Plugins/Experimental/ChaosFlesh/Source/ChaosFleshEngine/Private/ChaosFlesh/FleshComponent.cpp
--- a/Engine/Plugins/Experimental/ChaosFlesh/Source/ChaosFleshEngine/Private/ChaosFlesh/FleshComponent.cpp
+++ b/Engine/Plugins/Experimental/ChaosFlesh/Source/ChaosFleshEngine/Private/ChaosFlesh/FleshComponent.cpp
@@ -545,17 +545,22 @@ void UFleshComponent::DebugDrawSkeletalMeshBindingPositions() const
 
 	if (const UFleshAsset* RestAsset = GetRestCollection())
 	{
-		const USkeletalMesh* SkeletalMesh = TargetDeformationSkeleton ? TargetDeformationSkeleton : RestAsset->SkeletalMesh;
+		const USkeletalMesh* SkeletalMesh = 
+			GetRestCollection() && GetRestCollection()->TargetDeformationSkeleton ? 
+			GetRestCollection()->TargetDeformationSkeleton : RestAsset->SkeletalMesh;
 		if (SkeletalMesh)
 		{
 			TArray<bool> Influenced;
-			TArray<FVector> PosArray = GetSkeletalMeshBindingPositionsInternal(SkeletalMesh, &Influenced);
-			for (int i=0;i<PosArray.Num();i++)
+			TArray<FVector> PosArray = GetSkeletalMeshEmbeddedPositionsInternal(
+				ChaosDeformableBindingOption::ComponentPos, FTransform::Identity, "", &Influenced);
+
+			for (int i = 0; i < PosArray.Num(); i++)
 			{
 				const FVector& Pos = PosArray[i];
 				if (Influenced[i])
 				{
-					Chaos::FDebugDrawQueue::GetInstance().DrawDebugPoint(GetComponentTransform().TransformPosition(Pos), FColor::Red, true, 2.0f, SDPG_Foreground, 10);
+					Chaos::FDebugDrawQueue::GetInstance().DrawDebugPoint(
+						GetComponentTransform().TransformPosition(Pos), FColor::Yellow, true, 2.0f, SDPG_Foreground, 10);
 				}
 			}
 		}
@@ -565,7 +570,263 @@ void UFleshComponent::DebugDrawSkeletalMeshBindingPositions() const
 
 TArray<FVector> UFleshComponent::GetSkeletalMeshBindingPositions(const USkeletalMesh* InSkeletalMesh) const
 {
-	return GetSkeletalMeshBindingPositionsInternal(InSkeletalMesh);
+	return GetSkeletalMeshBindingPositionsInternal(InSkeletalMesh, nullptr);
+}
+
+TArray<FVector> UFleshComponent::GetSkeletalMeshEmbeddedPositions(
+	const ChaosDeformableBindingOption Format,
+	const FTransform TargetDeformationSkeletonOffset,
+	const FName TargetBone) const
+{
+	return GetSkeletalMeshEmbeddedPositionsInternal(Format, TargetDeformationSkeletonOffset, TargetBone, nullptr);
+}
+
+TArray<FVector> UFleshComponent::GetSkeletalMeshEmbeddedPositionsInternal(
+	const ChaosDeformableBindingOption Format,
+	const FTransform TargetDeformationSkeletonOffset, 
+	const FName TargetBone,
+	TArray<bool>* OutInfluence) const
+{
+	if (!GetRestCollection())
+	{
+		UE_LOG(LogFleshComponentInternal, Warning,
+			TEXT("'%s' - GetSkeletalMeshEmbeddedPositionsInternal - RestCollection is not set."),
+			*GetName());
+		return TArray<FVector>();
+	}
+	if(!GetRestCollection()->TargetDeformationSkeleton)
+	{
+		UE_LOG(LogFleshComponentInternal, Warning,
+			TEXT("'%s' - GetSkeletalMeshEmbeddedPositionsInternal - TargetDeformationSkeleton is not set on the flesh asset."),
+			*GetName());
+		return TArray<FVector>();
+	}
+
+	// Get sample points in the skel mesh's component space.  This code assumes that the
+	// skeletal mesh and the flesh asset are aligned in their respective local spaces.  If
+	// they're not aligned, then the TargetDeformationSkeletonOffset should be provided to
+	// put the skeletal mesh in the same place as the flesh mesh.
+	TArray<FVector> TransformPositions;
+	if (Format == ChaosDeformableBindingOption::WorldDelta || 
+		Format == ChaosDeformableBindingOption::ComponentDelta) // BoneDelta handled below
+	{
+		TArray<FTransform> ComponentPose;
+		Dataflow::Animation::GlobalTransforms(GetRestCollection()->TargetDeformationSkeleton->GetRefSkeleton(), ComponentPose);
+		TransformPositions.SetNumUninitialized(ComponentPose.Num());
+		if (TargetDeformationSkeletonOffset.Equals(FTransform::Identity))
+		{
+			for (int32 i = 0; i < ComponentPose.Num(); i++)
+			{
+				TransformPositions[i] = ComponentPose[i].GetTranslation();
+			}
+		}
+		else
+		{
+			for (int32 i = 0; i < ComponentPose.Num(); i++)
+			{
+				TransformPositions[i] =
+					TargetDeformationSkeletonOffset.TransformPosition(
+						ComponentPose[i].GetTranslation());
+			}
+		}
+	}
+	else
+	{
+		// If we aren't computing deltas, then we don't need the bone positions.  We
+		// only need to size the array to however many bones we have.
+		TransformPositions.SetNum(GetRestCollection()->TargetDeformationSkeleton->GetRefSkeleton().GetNum());
+	}
+
+	// Calculate their current embedded positions
+	TArray<FVector> EmbeddedPosComp = 
+		GetEmbeddedPositionsInternal(
+			TransformPositions, 
+			FName(GetRestCollection()->TargetDeformationSkeleton->GetName()), // for identifying the binding group
+			OutInfluence);
+
+	// World space
+	if (Format == ChaosDeformableBindingOption::WorldPos)
+	{
+		// Put component space points into world space
+		const FTransform ComponentXf = GetComponentTransform();
+		for (int32 i = 0; i < EmbeddedPosComp.Num(); i++)
+		{
+			EmbeddedPosComp[i] = ComponentXf.TransformPosition(EmbeddedPosComp[i]);
+		}
+	}
+	else if (Format == ChaosDeformableBindingOption::WorldDelta)
+	{
+		const FTransform ComponentXf = GetComponentTransform();
+		for (int32 i = 0; i < EmbeddedPosComp.Num(); i++)
+		{
+			EmbeddedPosComp[i] = EmbeddedPosComp[i] - TransformPositions[i];
+			EmbeddedPosComp[i] = ComponentXf.TransformVector(EmbeddedPosComp[i]);
+		}
+	}
+
+	// Component space
+	else if (Format == ChaosDeformableBindingOption::ComponentPos)
+	{
+		// Do nothing.
+	}
+	else if (Format == ChaosDeformableBindingOption::ComponentDelta)
+	{
+		for (int32 i = 0; i < EmbeddedPosComp.Num(); i++)
+		{
+			EmbeddedPosComp[i] = EmbeddedPosComp[i] - TransformPositions[i];
+		}
+	}
+
+	// Bone space
+	else
+	{
+		// Find the component that owns TargetDeformationSkeleton, so we can pull the 
+		// current animated bone transforms out of it.
+		//
+		// It's possible that TargetDeformationSkeleton is owned by a component on another 
+		// actor, in which case, we can't (easily) find it. If that becomes a desired use
+		// case, then we'll need to have a handle to that component, not just the asset.
+		TArray<USkeletalMeshComponent*> SkeletalMeshComponents;
+		GetOwner()->GetComponents<USkeletalMeshComponent>(SkeletalMeshComponents);
+		USkeletalMeshComponent* TargetDeformationSkeletalMeshComponent = nullptr;
+		for (int32 i = 0; i < SkeletalMeshComponents.Num(); i++)
+		{
+			if (SkeletalMeshComponents[i]->GetSkeletalMeshAsset() == GetRestCollection()->TargetDeformationSkeleton)
+			{
+				TargetDeformationSkeletalMeshComponent = SkeletalMeshComponents[i];
+				break;
+			}
+		}
+
+		if (TargetDeformationSkeletalMeshComponent)
+		{
+			// Get the current animated bone transforms and then their positions.
+			const TArray<FTransform>& TargetDeformationComponentTransforms = 
+				TargetDeformationSkeletalMeshComponent->GetComponentSpaceTransforms();
+			TArray<FVector> AnimTransformPositions;
+			AnimTransformPositions.SetNumUninitialized(TargetDeformationComponentTransforms.Num());
+
+			// Get their positions, and apply the offset if necessary.
+			if (!TargetDeformationSkeletonOffset.Equals(FTransform::Identity) && 
+				Format == ChaosDeformableBindingOption::BoneDelta)
+			{
+				for (int32 i = 0; i < TargetDeformationComponentTransforms.Num(); i++)
+				{
+					AnimTransformPositions[i] =
+						TargetDeformationSkeletonOffset.TransformPosition(
+							TargetDeformationComponentTransforms[i].GetTranslation());
+				}
+			}
+			else
+			{
+				for (int32 i = 0; i < TargetDeformationComponentTransforms.Num(); i++)
+				{
+					AnimTransformPositions[i] = TargetDeformationComponentTransforms[i].GetTranslation();
+				}
+			}
+
+			// Find the transform index of 'TargetBone'
+			FTransform BoneToComponentXf = FTransform::Identity;
+			const int32 BoneIndex = GetRestCollection()->TargetDeformationSkeleton->GetRefSkeleton().FindBoneIndex(TargetBone);
+			if (TargetDeformationComponentTransforms.IsValidIndex(BoneIndex))
+			{
+				BoneToComponentXf = FTransform(TargetDeformationComponentTransforms[BoneIndex].ToMatrixWithScale().Inverse());
+			}
+			else
+			{
+				UE_LOG(LogFleshComponentInternal, Warning, 
+					TEXT("'%s' - Failed to find a valid bone index (got %d) for bone name '%s' in TargetDeformationSkeleton '%s' "
+					"corresponding to SkeletalMeshComponent '%s', which has %d bones."), 
+					*GetName(),
+					BoneIndex, *TargetBone.ToString(), *GetRestCollection()->TargetDeformationSkeleton.GetName(),
+					*TargetDeformationSkeletalMeshComponent->GetName(), TargetDeformationComponentTransforms.Num());
+				return TArray<FVector>();
+			}
+
+			// Compute the return values
+			if (Format == ChaosDeformableBindingOption::BonePos)
+			{
+				if (!BoneToComponentXf.Equals(FTransform::Identity))
+				{
+					for (int32 i = 0; i < EmbeddedPosComp.Num(); i++)
+					{
+						EmbeddedPosComp[i] = BoneToComponentXf.TransformPosition(EmbeddedPosComp[i]);
+					}
+				}
+			}
+			else if (Format == ChaosDeformableBindingOption::BoneDelta)
+			{
+				if (!BoneToComponentXf.Equals(FTransform::Identity))
+				{
+					for (int32 i = 0; i < EmbeddedPosComp.Num(); i++)
+					{
+						EmbeddedPosComp[i] = EmbeddedPosComp[i] - AnimTransformPositions[i];
+						EmbeddedPosComp[i] = BoneToComponentXf.TransformVector(EmbeddedPosComp[i]);
+					}
+				}
+				else
+				{
+					for (int32 i = 0; i < EmbeddedPosComp.Num(); i++)
+					{
+						EmbeddedPosComp[i] = EmbeddedPosComp[i] - AnimTransformPositions[i];
+					}
+				}
+			}
+		}
+		else
+		{
+			UE_LOG(LogFleshComponentInternal, Warning, 
+				TEXT("'%s' - Failed to find SkeletalMeshComponent for TargetDeformationSkeleton '%s'."),
+				*GetName(),
+				*GetRestCollection()->TargetDeformationSkeleton->GetName());
+			return TArray<FVector>();
+		}
+	}
+	return EmbeddedPosComp;
+}
+
+TArray<FVector> UFleshComponent::GetEmbeddedPositionsInternal(
+	const TArray<FVector>& InPositions, 
+	const FName SkeletalMeshName,
+	TArray<bool>* OutInfluence) const
+{
+	auto UEVert3d = [](FVector3f V) { return FVector3d(V.X, V.Y, V.Z); };
+	auto UEVert4d = [](FVector4f V) { return FVector4d(V.X, V.Y, V.Z, V.W); };
+
+	TArray<FVector> OutPositions = InPositions;
+
+	if (const UFleshAsset* RestAsset = GetRestCollection())
+	{
+		if (const FFleshCollection* Rest = RestAsset->GetCollection())
+		{
+			GeometryCollection::Facades::FTetrahedralSkeletalBindings TetBindings(*Rest);
+
+			const TManagedArray<int32>* TetrahedronStart = Rest->FindAttribute<int32>(FTetrahedralCollection::TetrahedronStartAttribute, FGeometryCollection::GeometryGroup);
+			const TManagedArray<FVector3f>* RestVerts = Rest->FindAttributeTyped<FVector3f>("Vertex", FGeometryCollection::VerticesGroup);
+			const TManagedArray<FVector3f>* Verts = GetDynamicCollection() ? GetDynamicCollection()->FindPositions() : RestVerts;
+
+			if (ensure(Verts != nullptr) && TetrahedronStart)
+			{
+				if (OutInfluence != nullptr)
+					OutInfluence->Init(false, OutPositions.Num());
+				for (int32 TetMeshIdx = 0; TetMeshIdx < TetrahedronStart->Num(); TetMeshIdx++)
+				{
+					FString MeshBindingsName =
+						GeometryCollection::Facades::FTetrahedralSkeletalBindings::GenerateMeshGroupName(TetMeshIdx, SkeletalMeshName);
+					if (!TetBindings.CalculateBindings(MeshBindingsName, Verts->GetConstArray(), OutPositions, OutInfluence))
+					{
+						UE_LOG(LogFleshComponentInternal, Warning,
+							TEXT("'%s' - GetEmbeddedPositionsInternal - Failed to find mesh bindings for skeletal mesh '%s'"),
+							*GetName(),
+							*SkeletalMeshName.ToString());
+					}
+				}
+
+			}
+
+		}
+	}
+	return OutPositions;
 }
 
 TArray<FVector> UFleshComponent::GetSkeletalMeshBindingPositionsInternal(const USkeletalMesh* InSkeletalMesh, TArray<bool>* OutInfluence) const
@@ -588,6 +849,7 @@ TArray<FVector> UFleshComponent::GetSkeletalMeshBindingPositionsInternal(const U
 
 				if (ensure(Verts != nullptr) && TetrahedronStart)
 				{
+					// Component relative transforms, not world.
 					TArray<FTransform> ComponentPose;
 					Dataflow::Animation::GlobalTransforms(InSkeletalMesh->GetRefSkeleton(), ComponentPose);
 
@@ -609,11 +871,3 @@ TArray<FVector> UFleshComponent::GetSkeletalMeshBindingPositionsInternal(const U
 	}
 	return TransformPositions;
 }
-
-
-
-
-
-
-
-
diff --git a/Engine/Plugins/Experimental/ChaosFlesh/Source/ChaosFleshEngine/Public/ChaosFlesh/FleshAsset.h b/Engine/Plugins/Experimental/ChaosFlesh/Source/ChaosFleshEngine/Public/ChaosFlesh/FleshAsset.h
--- a/Engine/Plugins/Experimental/ChaosFlesh/Source/ChaosFleshEngine/Public/ChaosFlesh/FleshAsset.h
+++ b/Engine/Plugins/Experimental/ChaosFlesh/Source/ChaosFleshEngine/Public/ChaosFlesh/FleshAsset.h
@@ -113,6 +113,13 @@ public:
 	UPROPERTY(EditAnywhere, Category = "Animation")
 	TObjectPtr<USkeleton> Skeleton;
 
+	/**
+	* Skeleton to use with \c GetSkeletalMeshEmbeddedPositions() on the flesh component. 
+	* Bindings for this skeletal mesh must be stored in the rest collection.
+	*/
+	UPROPERTY(EditAnywhere, BlueprintReadWrite, Category = "Rendering")
+	TObjectPtr<USkeletalMesh> TargetDeformationSkeleton;
+
 	//
 	// SkeletalMesh
 	//
diff --git a/Engine/Plugins/Experimental/ChaosFlesh/Source/ChaosFleshEngine/Public/ChaosFlesh/FleshComponent.h b/Engine/Plugins/Experimental/ChaosFlesh/Source/ChaosFleshEngine/Public/ChaosFlesh/FleshComponent.h
--- a/Engine/Plugins/Experimental/ChaosFlesh/Source/ChaosFleshEngine/Public/ChaosFlesh/FleshComponent.h
+++ b/Engine/Plugins/Experimental/ChaosFlesh/Source/ChaosFleshEngine/Public/ChaosFlesh/FleshComponent.h
@@ -16,6 +16,21 @@ class FFleshCollection;
 class ADeformableSolverActor;
 class UDeformableSolverComponent;
 
+/**
+*  Options for binding positions query.
+*/
+UENUM()
+enum ChaosDeformableBindingOption : uint8
+{
+	WorldPos		UMETA(DisplayName = "World Positions"),
+	WorldDelta		UMETA(DisplayName = "World Deltas"),
+	ComponentPos    UMETA(DisplayName = "Component Positions"),
+	ComponentDelta  UMETA(DisplayName = "Component Deltas"),
+	BonePos			UMETA(DisplayName = "Bone Positions"),
+	BoneDelta		UMETA(DisplayName = "Bone Deltas"),
+};
+
+
 /**
 *	FleshComponent
 */
@@ -29,9 +44,6 @@ public:
 
 	~UFleshComponent();
 
-	UFUNCTION(BlueprintCallable, Category = "Physics")
-	TArray<FVector> GetSkeletalMeshBindingPositions(const USkeletalMesh* InSkeletalMesh) const;
-
 	/** USceneComponent Interface */
 	virtual void BeginPlay() override;
 	virtual void OnRegister() override;
@@ -61,16 +73,32 @@ public:
 	UPROPERTY()
 	TObjectPtr<UProceduralMeshComponent> Mesh;
 
-	UPROPERTY(EditAnywhere, BlueprintReadWrite, Category = "Rendering")
-	TObjectPtr<USkeletalMesh> TargetDeformationSkeleton;
+	/** @deprecated Use GetSkeletalMeshEmbeddedPositions() instead. */
+	UFUNCTION(BlueprintCallable, Category = "Physics", meta = (DeprecatedFunction, DeprecationMessage = "Use GetSkeletalMeshEmbeddedPositions() instead."))
+	TArray<FVector> GetSkeletalMeshBindingPositions(const USkeletalMesh* InSkeletalMesh) const;
+
+	/**
+	* Get the current positions of the transformation hierarchy from \c TargetDeformationSkeleton,
+	* deformed by the tetrahedral mesh.  Results can be in world space postions/deltas, component space
+	* positions/deltas, or bone space positions/deltas.  If a bone space is desired \p TargetBone
+	* must indicate which bone to use. TargetDeformationSkeletonOffset is an offset transform that moves 
+	* the \c TargetDeformationSkeleton to be co-located with the flesh mesh.
+	*/
+	UFUNCTION(BlueprintCallable, Category = "Physics")
+	TArray<FVector> GetSkeletalMeshEmbeddedPositions(const ChaosDeformableBindingOption Format, const FTransform TargetDeformationSkeletonOffset, const FName TargetBone = "") const;
 
+	/** 
+	* Offset transform that moves the \c TargetDeformationSkeleton to be co-located with the flesh mesh. 
+	*/
+//	UPROPERTY(EditAnywhere, BlueprintReadWrite, Category = "Rendering")
+//	FTransform TargetDeformationSkeletonOffset;
 
 private:
-	// FleshAsset that describes the simulation rest state.
+	/** FleshAsset that describes the simulation rest state. */
 	UPROPERTY(EditAnywhere, Category = "ChaosPhysics")
 	TObjectPtr<const UFleshAsset> RestCollection;
 
-	// Current simulation state.
+	/** Current simulation state. */
 	UPROPERTY()
 	TObjectPtr<UFleshDynamicAsset> DynamicCollection;
 
@@ -78,18 +106,22 @@ private:
 	// Sim Space
 	//
 
-	// Space the simulation will run in.
+	/** Space the simulation will run in. */
 	UPROPERTY(EditAnywhere, Category = "ChaosDeformable")
 	TEnumAsByte<ChaosDeformableSimSpace> SimSpace = ChaosDeformableSimSpace::World;
 
-	// Bone from the associated skeletal mesh (indicated by RestCollection.TargetSkeletalMesh) to use as the space the sim runs in.
+	/** 
+	* Bone from the associated skeletal mesh (indicated by RestCollection.TargetSkeletalMesh) to use as 
+	* the space the sim runs in.
+	*/
 	UPROPERTY(EditAnywhere, Category = "ChaosDeformable", meta = (GetOptions = "GetSimSpaceBoneNameOptions", EditCondition = "SimSpace == ChaosDeformableSimSpace::Bone"))
 	FName SimSpaceBoneName;
 
+	/** The skeletal mesh to use pull the \c SimSpaceBoneName from. */
 	UPROPERTY()
 	TObjectPtr<USkeletalMesh> SimSpaceSkeletalMesh;
 
-	// Returns a list of bone names from the currently selected skeletal mesh.
+	/* Returns a list of bone names from the currently selected skeletal mesh. */
 	UFUNCTION(CallInEditor)
 	TArray<FString> GetSimSpaceBoneNameOptions() const;
 
@@ -123,6 +155,8 @@ private:
 
 	FTransform PrevTransform = FTransform::Identity;
 
+	TArray<FVector> GetSkeletalMeshEmbeddedPositionsInternal(const ChaosDeformableBindingOption Format, const FTransform TargetDeformationSkeletonOffset, const FName TargetBone = "", TArray<bool>* OutInfluence = nullptr) const;
+	TArray<FVector> GetEmbeddedPositionsInternal(const TArray<FVector>& InPositions, const FName SkeletalMeshName, TArray<bool>* OutInfluence = nullptr) const;
 	TArray<FVector> GetSkeletalMeshBindingPositionsInternal(const USkeletalMesh* InSkeletalMesh, TArray<bool>* OutInfluence = nullptr) const;
 	void DebugDrawSkeletalMeshBindingPositions() const;
 
diff --git a/Engine/Source/Runtime/Experimental/Chaos/Private/GeometryCollection/Facades/CollectionTetrahedralSkeletalBindingsFacade.cpp b/Engine/Source/Runtime/Experimental/Chaos/Private/GeometryCollection/Facades/CollectionTetrahedralSkeletalBindingsFacade.cpp
--- a/Engine/Source/Runtime/Experimental/Chaos/Private/GeometryCollection/Facades/CollectionTetrahedralSkeletalBindingsFacade.cpp
+++ b/Engine/Source/Runtime/Experimental/Chaos/Private/GeometryCollection/Facades/CollectionTetrahedralSkeletalBindingsFacade.cpp
@@ -101,46 +101,53 @@ namespace GeometryCollection::Facades
 	}
 
 
-	void FTetrahedralSkeletalBindings::CalculateBindings(const FString & InKey, const TArray<FVector3f>& InVertices, TArray<FVector>& OutPosition, TArray<bool>* OutInfluence) const
+	bool FTetrahedralSkeletalBindings::CalculateBindings(const FString & InKey, const TArray<FVector3f>& InVertices, TArray<FVector>& OutPosition, TArray<bool>* OutInfluence) const
 	{
 		auto UEVert3d = [](FVector3f V) { return FVector3d(V.X, V.Y, V.Z); };
 
-		if (IsValid() && InVertices.Num())
+		if (!IsValid())
+		{
+			return false;
+		}
+		if(InVertices.Num())
 		{
 			int32 GroupIndex = MeshIdAttribute.Get().Find(InKey);
-			if (GroupIndex != INDEX_NONE)
+			if (GroupIndex == INDEX_NONE)
+			{
+				return false;
+			}
+
+			const TManagedArray<int32>& MeshGroupIndex = MeshGroupIndexAttribute.Get();
+			const TManagedArray<int32>& TetrahedronIndex = TetrahedronIndexAttribute.Get();
+			const TManagedArray<FVector4f>& Weights = WeightsAttribute.Get();
+			const TManagedArray<int32>& SkeletalIndices = SkeletalIndexAttribute.Get();
+			const TManagedArray<FIntVector4>& Tetrahedron = TetrahedronAttribute.Get();
+
+			int32 NumBindings = MeshGroupIndex.Num();
+			for (int32 BindingIndex = 0; BindingIndex < MeshGroupIndex.Num(); BindingIndex++)
 			{
-				const TManagedArray<int32>& MeshGroupIndex = MeshGroupIndexAttribute.Get();
-				const TManagedArray<int32>& TetrahedronIndex = TetrahedronIndexAttribute.Get();
-				const TManagedArray<FVector4f>& Weights = WeightsAttribute.Get();
-				const TManagedArray<int32>& SkeletalIndices = SkeletalIndexAttribute.Get();
-				const TManagedArray<FIntVector4>& Tetrahedron = TetrahedronAttribute.Get();
-
-				int32 NumBindings = MeshGroupIndex.Num();
-				for (int32 BindingIndex = 0; BindingIndex < MeshGroupIndex.Num(); BindingIndex++)
+				if (MeshGroupIndex[BindingIndex] == GroupIndex)
 				{
-					if (MeshGroupIndex[BindingIndex] == GroupIndex)
+					int32 TetIndex = TetrahedronIndex[BindingIndex];
+					if (0 <= TetIndex && TetIndex < TetrahedronAttribute.Num())
 					{
-						int32 TetIndex = TetrahedronIndex[BindingIndex];
-						if (0 <= TetIndex && TetIndex < TetrahedronAttribute.Num())
+						const FIntVector4& Tet = Tetrahedron[TetIndex];
+						int32 SkeletonIndex = SkeletalIndices[BindingIndex];
+						if (0 <= SkeletonIndex && SkeletonIndex < OutPosition.Num())
 						{
-							const FIntVector4& Tet = Tetrahedron[TetIndex];
-							int32 SkeletonIndex = SkeletalIndices[BindingIndex];
-							if (0 <= SkeletonIndex && SkeletonIndex < OutPosition.Num())
-							{
-								const FVector4 Weight(Weights[BindingIndex][0], Weights[BindingIndex][1], Weights[BindingIndex][2], Weights[BindingIndex][3]);
-								FVector X0 = InVertices.IsValidIndex(Tet[0]) ? UEVert3d(InVertices[Tet[0]]) : FVector();
-								FVector X1 = InVertices.IsValidIndex(Tet[0]) ? UEVert3d(InVertices[Tet[1]]) : FVector();
-								FVector X2 = InVertices.IsValidIndex(Tet[0]) ? UEVert3d(InVertices[Tet[2]]) : FVector();
-								FVector X3 = InVertices.IsValidIndex(Tet[0]) ? UEVert3d(InVertices[Tet[3]]) : FVector();
-								OutPosition[SkeletonIndex] = Weight[0] * X0 + Weight[1] * X1 + Weight[2] * X2 + Weight[3] * X3;
-								if (OutInfluence!=nullptr) if(SkeletonIndex<OutInfluence->Num()) (*OutInfluence)[SkeletonIndex] = true;
-							}
+							const FVector4 Weight(Weights[BindingIndex][0], Weights[BindingIndex][1], Weights[BindingIndex][2], Weights[BindingIndex][3]);
+							FVector X0 = InVertices.IsValidIndex(Tet[0]) ? UEVert3d(InVertices[Tet[0]]) : FVector();
+							FVector X1 = InVertices.IsValidIndex(Tet[0]) ? UEVert3d(InVertices[Tet[1]]) : FVector();
+							FVector X2 = InVertices.IsValidIndex(Tet[0]) ? UEVert3d(InVertices[Tet[2]]) : FVector();
+							FVector X3 = InVertices.IsValidIndex(Tet[0]) ? UEVert3d(InVertices[Tet[3]]) : FVector();
+							OutPosition[SkeletonIndex] = Weight[0] * X0 + Weight[1] * X1 + Weight[2] * X2 + Weight[3] * X3;
+							if (OutInfluence != nullptr) if (SkeletonIndex < OutInfluence->Num()) (*OutInfluence)[SkeletonIndex] = true;
 						}
 					}
 				}
 			}
 		}
+		return true;
 	}
 };
 
diff --git a/Engine/Source/Runtime/Experimental/Chaos/Public/GeometryCollection/Facades/CollectionTetrahedralSkeletalBindingsFacade.h b/Engine/Source/Runtime/Experimental/Chaos/Public/GeometryCollection/Facades/CollectionTetrahedralSkeletalBindingsFacade.h
--- a/Engine/Source/Runtime/Experimental/Chaos/Public/GeometryCollection/Facades/CollectionTetrahedralSkeletalBindingsFacade.h
+++ b/Engine/Source/Runtime/Experimental/Chaos/Public/GeometryCollection/Facades/CollectionTetrahedralSkeletalBindingsFacade.h
@@ -57,7 +57,9 @@ namespace GeometryCollection::Facades
 		//
 		static FString GenerateMeshGroupName(const int32 TetMeshIdx, const FName& SkeletalMeshName, const int32 LOD = INDEX_NONE);
 		void SetBindings(const FString& InMeshGroupIndex, const TArray<int32>& InTetrahedronIndex, const TArray<FVector4f>& WeightsIn, const TArray<int32>& InSkeletalIndex);
-		void CalculateBindings(const FString& InKey, const TArray<FVector3f>& InVertices, TArray<FVector>& OutPosition, TArray<bool>* OutInfluence = nullptr) const;
+		
+		/** Returns \c false if bindings data was not found, or there was another error; \c true otherwise. */
+		bool CalculateBindings(const FString& InKey, const TArray<FVector3f>& InVertices, TArray<FVector>& OutPosition, TArray<bool>* OutInfluence = nullptr) const;
 		
 	private:
 		TManagedArrayAccessor<FString> MeshIdAttribute;
