ChaosFlesh : Multishape collision support

diff --git a/Engine/Plugins/Experimental/ChaosFlesh/Source/ChaosFleshEngine/Private/ChaosFlesh/ChaosDeformableCollisionsComponent.cpp b/Engine/Plugins/Experimental/ChaosFlesh/Source/ChaosFleshEngine/Private/ChaosFlesh/ChaosDeformableCollisionsComponent.cpp
--- a/Engine/Plugins/Experimental/ChaosFlesh/Source/ChaosFleshEngine/Private/ChaosFlesh/ChaosDeformableCollisionsComponent.cpp
+++ b/Engine/Plugins/Experimental/ChaosFlesh/Source/ChaosFleshEngine/Private/ChaosFlesh/ChaosDeformableCollisionsComponent.cpp
@@ -64,7 +64,7 @@ template<> class Chaos::TSphere<float, 3>;
 
 UDeformableCollisionsComponent::FDataMapValue 
 UDeformableCollisionsComponent::NewDeformableData()
-{ 
+{
 	TArray < Chaos::Softs::FCollisionObjectAddedBodies> AddedBodiesData;
 	TArray < Chaos::Softs::FCollisionObjectRemovedBodies> RemovedBodiesData;
 	TArray < Chaos::Softs::FCollisionObjectUpdatedBodies> UpdateBodiesData;
@@ -83,23 +83,34 @@ UDeformableCollisionsComponent::NewDeformableData()
 					FVector Scale = P.GetScale3D();
 					P.RemoveScaling();
 
-					if (BodySetup->AggGeom.SphereElems.Num())
+					for( int8 Index = 0; Index < BodySetup->AggGeom.SphereElems.Num(); Index++)
 					{
-						FKSphereElem& S = BodySetup->AggGeom.SphereElems[0];
-						Geometry = new ::Chaos::TSphere<Chaos::FReal, 3>(S.Center, S.Radius * Scale.GetMax());
+						FKSphereElem& S = BodySetup->AggGeom.SphereElems[Index];
+						FTransform ShapeTransform = FTransform(Scale * S.Center) * P;
+						Geometry = new ::Chaos::TSphere<Chaos::FReal, 3>(FVector(0), S.Radius * Scale.GetMax());
+						AddedBodiesData.Add(Chaos::Softs::FCollisionObjectAddedBodies(
+							{ CollisionBody, Chaos::Softs::ERigidCollisionShapeType::Sphere, Index },
+							ShapeTransform, "", Geometry));
 					}
-					else if (BodySetup->AggGeom.BoxElems.Num())
+
+					for (int8 Index = 0; Index < BodySetup->AggGeom.BoxElems.Num(); Index++)
 					{
-						FKBoxElem& B = BodySetup->AggGeom.BoxElems[0];
+						FKBoxElem& B = BodySetup->AggGeom.BoxElems[Index];
+						FTransform ShapeTransform = FTransform(Scale * B.Center) * P;
 						Geometry = new TBox<Chaos::FReal, 3>(
-							B.Center - 0.5 * Scale * FVector(B.X, B.Y, B.Z),
-							B.Center + 0.5 * Scale * FVector(B.X, B.Y, B.Z));
+							-0.5 * Scale * FVector(B.X, B.Y, B.Z),
+							0.5 * Scale * FVector(B.X, B.Y, B.Z));
+						AddedBodiesData.Add(Chaos::Softs::FCollisionObjectAddedBodies(
+							{ CollisionBody, Chaos::Softs::ERigidCollisionShapeType::Box, Index },
+							ShapeTransform, "", Geometry));
 					}
-					else if (BodySetup->AggGeom.ConvexElems.Num())
+
+					for (int8 Index = 0; Index < BodySetup->AggGeom.ConvexElems.Num(); Index++)
 					{
-						FKConvexElem& C = BodySetup->AggGeom.ConvexElems[0];
+						FKConvexElem& C = BodySetup->AggGeom.ConvexElems[Index];
 						if (C.VertexData.Num())
 						{
+							FTransform ShapeTransform = C.GetTransform() * P;
 							TSharedPtr<Chaos::FConvex, ESPMode::ThreadSafe> ConvexMesh = C.GetChaosConvexMesh();	
 							TArray<FConvex::FVec3Type> Vertices;
 							Vertices.SetNum(C.VertexData.Num());
@@ -108,17 +119,11 @@ UDeformableCollisionsComponent::NewDeformableData()
 								Vertices[i] = C.VertexData[i] * Scale;
 							}
 							Geometry = new FConvex(Vertices, 0.f);
+							AddedBodiesData.Add(Chaos::Softs::FCollisionObjectAddedBodies(
+								{ CollisionBody, Chaos::Softs::ERigidCollisionShapeType::Convex, Index},
+								ShapeTransform, "", Geometry));
 						}
 					}
-					//else if (BodySetup->ChaosTriMeshes.Num())
-					//{
-					//	TArray<TSharedPtr<Chaos::FTriangleMeshImplicitObject, ESPMode::ThreadSafe>> ChaosTriMeshes;
-					//	Geometry = BodySetup->ChaosTriMeshes[0]->CopySlow().Release();
-					//}
-					if (Geometry)
-					{
-						AddedBodiesData.Add(Chaos::Softs::FCollisionObjectAddedBodies(CollisionBody, P, "", Geometry));
-					}
 				}
 			}
 		}
@@ -129,7 +134,7 @@ UDeformableCollisionsComponent::NewDeformableData()
 	{
 		if (RemovedBody)
 		{
-			RemovedBodiesData.Add({ RemovedBody });
+			RemovedBodiesData.Add({  {RemovedBody, Chaos::Softs::ERigidCollisionShapeType::Unknown, INDEX_NONE} });
 		}
 	}
 	AddedBodies.Empty();
@@ -140,7 +145,7 @@ UDeformableCollisionsComponent::NewDeformableData()
 		if (CollisionBody)
 		{
 			FTransform P = CollisionBody->GetComponentToWorld();
-			UpdateBodiesData.Add({ CollisionBody,P });
+			UpdateBodiesData.Add({ {CollisionBody, Chaos::Softs::ERigidCollisionShapeType::Unknown, INDEX_NONE},P });
 		}
 	}
 
diff --git a/Engine/Source/Runtime/Experimental/Chaos/Private/Chaos/Deformable/ChaosDeformableSolver.cpp b/Engine/Source/Runtime/Experimental/Chaos/Private/Chaos/Deformable/ChaosDeformableSolver.cpp
--- a/Engine/Source/Runtime/Experimental/Chaos/Private/Chaos/Deformable/ChaosDeformableSolver.cpp
+++ b/Engine/Source/Runtime/Experimental/Chaos/Private/Chaos/Deformable/ChaosDeformableSolver.cpp
@@ -499,24 +499,38 @@ namespace Chaos::Softs
 					{
 						if (AddBody.Shapes)
 						{
-							typedef TPlane<Chaos::Softs::FSolverReal, 3> GeomType;
-							int32 Index = Evolution->AddCollisionParticle(INDEX_NONE, true);
-							int32 ViewIndex = Evolution->CollisionParticlesActiveView().GetRanges().Num()-1;
-							Evolution->CollisionParticles().X(Index) = AddBody.Transform.GetTranslation();
-							Evolution->CollisionParticles().R(Index) = AddBody.Transform.GetRotation();
-							TUniquePtr<FImplicitObject> UniquePtr(AddBody.Shapes); AddBody.Shapes = nullptr;
-							Evolution->CollisionParticles().SetDynamicGeometry(Index, MoveTemp(UniquePtr));
-							Proxy.CollisionBodies.Add(AddBody.BodyId, FCollisionObjectParticleHandel( Index,ViewIndex,AddBody.Transform ));
+							if (ensure(!Proxy.CollisionBodies.Contains(AddBody.Key)))
+							{
+								int32 Index = Evolution->AddCollisionParticle(INDEX_NONE, true);
+								int32 ViewIndex = Evolution->CollisionParticlesActiveView().GetRanges().Num() - 1;
+								Evolution->CollisionParticles().X(Index) = AddBody.Transform.GetTranslation();
+								Evolution->CollisionParticles().R(Index) = AddBody.Transform.GetRotation();
+								TUniquePtr<FImplicitObject> UniquePtr(AddBody.Shapes); AddBody.Shapes = nullptr;
+								Evolution->CollisionParticles().SetDynamicGeometry(Index, MoveTemp(UniquePtr));
+								Proxy.CollisionBodies.Add(AddBody.Key, FCollisionObjectParticleHandel(Index, ViewIndex, AddBody.Transform));
+							}
 						}
 					}
+
+					TArray<FCollisionObjectKey> KeysToRemove;
 					for (auto& RemovedBody : CollisionsInputBuffer->Removed)
 					{
-						if (Proxy.CollisionBodies.Contains(RemovedBody.BodyId))
+						for (auto& CollisionBodyPair : Proxy.CollisionBodies)
+						{
+							if ((void*)CollisionBodyPair.Key.Get<0>() == (void*)RemovedBody.Key.Get<0>())
+							{
+								KeysToRemove.Add(CollisionBodyPair.Key);
+							}
+						}
+					}
+					for (auto& KeyToRemove : KeysToRemove)
+					{
+						if (ensure(Proxy.CollisionBodies.Contains(KeyToRemove)))
 						{
-							int32 ParticleIndex = Proxy.CollisionBodies[RemovedBody.BodyId].ParticleIndex;
-							int32 ViewIndex = Proxy.CollisionBodies[RemovedBody.BodyId].ActiveViewIndex;
+							int32 ParticleIndex = Proxy.CollisionBodies[KeyToRemove].ParticleIndex;
+							int32 ViewIndex = Proxy.CollisionBodies[KeyToRemove].ActiveViewIndex;
 							Evolution->RemoveCollisionParticle(ParticleIndex, ViewIndex);
-							Proxy.CollisionBodies.Remove(RemovedBody.BodyId);
+							Proxy.CollisionBodies.Remove(KeyToRemove);
 						}
 					}
 				}
diff --git a/Engine/Source/Runtime/Experimental/Chaos/Public/Chaos/Deformable/ChaosDeformableCollisionsProxy.h b/Engine/Source/Runtime/Experimental/Chaos/Public/Chaos/Deformable/ChaosDeformableCollisionsProxy.h
--- a/Engine/Source/Runtime/Experimental/Chaos/Public/Chaos/Deformable/ChaosDeformableCollisionsProxy.h
+++ b/Engine/Source/Runtime/Experimental/Chaos/Public/Chaos/Deformable/ChaosDeformableCollisionsProxy.h
@@ -7,19 +7,32 @@
 
 namespace Chaos::Softs
 {
+	enum ERigidCollisionShapeType
+	{
+		Sphere,
+		Box,
+		Sphyl,
+		Convex,
+		TaperedCapsule,
+		LevelSet,
+
+		Unknown
+	};
+
+	typedef TTuple< const UObject*, ERigidCollisionShapeType, int8> FCollisionObjectKey;
+
 	struct CHAOS_API FCollisionObjectAddedBodies
 	{
-		FCollisionObjectAddedBodies( const UObject* InBodyId = nullptr,
+		FCollisionObjectAddedBodies( FCollisionObjectKey InKey = FCollisionObjectKey(),
 			FTransform InTransform = FTransform::Identity,
 			FString InType = "",
 			FImplicitObject* InShapes = nullptr)
-			: BodyId(InBodyId)
+			: Key(InKey)
 			, Transform(InTransform)
 			, Type(InType)
 			, Shapes(InShapes) {}
 
-	
-		const UObject* BodyId = nullptr;
+		FCollisionObjectKey Key = FCollisionObjectKey();
 		FTransform Transform = FTransform::Identity;
 		FString Type = "";
 		FImplicitObject* Shapes = nullptr;
@@ -27,12 +40,12 @@ namespace Chaos::Softs
 
 	struct CHAOS_API FCollisionObjectRemovedBodies
 	{
-		const UObject* BodyId = nullptr;
+		FCollisionObjectKey Key = FCollisionObjectKey();
 	};
 
 	struct CHAOS_API FCollisionObjectUpdatedBodies
 	{
-		const UObject* BodyId = nullptr;
+		FCollisionObjectKey Key = FCollisionObjectKey();
 		FTransform Transform = FTransform::Identity;
 	};
 
@@ -87,7 +100,9 @@ namespace Chaos::Softs
 
 		TArray<FCollisionObjectAddedBodies> CollisionObjectsToAdd;
 		TArray< FCollisionObjectRemovedBodies > CollisionObjectsToRemove;
-		TMap<const UObject*, FCollisionObjectParticleHandel > CollisionBodies;
+		TMap< FCollisionObjectKey, FCollisionObjectParticleHandel > CollisionBodies;
 	};
 
-}// namespace Chaos::Softs
\ No newline at end of file
+}// namespace Chaos::Softs
+
+
