ChaosFlesh : Add island computation node, point triangle position target node and debug draw support in XPBD weak constraints.

diff --git a/Engine/Plugins/Experimental/ChaosFlesh/Source/ChaosFleshNodes/Private/Dataflow/ChaosFleshEngineAssetNodes.cpp b/Engine/Plugins/Experimental/ChaosFlesh/Source/ChaosFleshNodes/Private/Dataflow/ChaosFleshEngineAssetNodes.cpp
--- a/Engine/Plugins/Experimental/ChaosFlesh/Source/ChaosFleshNodes/Private/Dataflow/ChaosFleshEngineAssetNodes.cpp
+++ b/Engine/Plugins/Experimental/ChaosFlesh/Source/ChaosFleshNodes/Private/Dataflow/ChaosFleshEngineAssetNodes.cpp
@@ -23,6 +23,7 @@ namespace Dataflow
 		DATAFLOW_NODE_REGISTER_CREATION_FACTORY(FFleshAssetTerminalDataflowNode);
 		DATAFLOW_NODE_REGISTER_CREATION_FACTORY(FSetFleshDefaultPropertiesNode);
 		DATAFLOW_NODE_REGISTER_CREATION_FACTORY(FComputeFiberFieldNode);
+		DATAFLOW_NODE_REGISTER_CREATION_FACTORY(FComputeIslandsNode);
 	}
 }
 
@@ -377,3 +378,53 @@ FComputeFiberFieldNode::ComputeFiberField(
 		Tolerance);
 	return Directions;
 }
+
+
+void FComputeIslandsNode::Evaluate(Dataflow::FContext& Context, const FDataflowOutput* Out) const
+{
+	if (Out->IsA<FManagedArrayCollection>(&Collection))
+	{
+		FManagedArrayCollection InCollection = GetValue<FManagedArrayCollection>(Context, &Collection);
+
+		TManagedArray<int32>& ParticleComponentIndex = InCollection.AddAttribute<int32>("ComponentIndex", FGeometryCollection::VerticesGroup);
+		TManagedArray<FIntVector4>* Elements = InCollection.FindAttribute<FIntVector4>(
+			FTetrahedralCollection::TetrahedronAttribute, FTetrahedralCollection::TetrahedralGroup);
+
+
+		if (InCollection.HasAttributes({
+			MType<int32>("ComponentIndex", FGeometryCollection::VerticesGroup) }) && Elements)
+		{
+
+			int32 VertsNum = InCollection.NumElements(FGeometryCollection::VerticesGroup);
+			int32 TetsNum = InCollection.NumElements(FTetrahedralCollection::TetrahedralGroup);
+			if (VertsNum && TetsNum)
+			{
+				TArray<TArray<int32>> ConnectedComponents;
+				Chaos::Utilities::FindConnectedRegions(Elements->GetConstArray(), ConnectedComponents);
+				TManagedArray<int32>& ComponentIndex = InCollection.ModifyAttribute<int32>("ComponentIndex", FGeometryCollection::VerticesGroup);
+
+				ComponentIndex.Fill(-1);
+
+				for (int32 i = 0; i < ConnectedComponents.Num(); i++)
+				{
+					for (int32 j = 0; j < ConnectedComponents[i].Num(); j++)
+					{
+						int32 ElementIndex = ConnectedComponents[i][j];
+						for (int32 ie = 0; ie < 4; ie++)
+						{
+							int32 ParticleIndex = (*Elements)[ElementIndex][ie];
+							if (ComponentIndex[ParticleIndex] == -1)
+							{
+								ComponentIndex[ParticleIndex] = i;
+							}
+						}
+					}
+				}
+			}
+		}
+		
+		SetValue<FManagedArrayCollection>(Context, InCollection, &Collection);
+	}
+}
+
+
diff --git a/Engine/Plugins/Experimental/ChaosFlesh/Source/ChaosFleshNodes/Private/Dataflow/ChaosFleshPositionTargetInitializationNodes.cpp b/Engine/Plugins/Experimental/ChaosFlesh/Source/ChaosFleshNodes/Private/Dataflow/ChaosFleshPositionTargetInitializationNodes.cpp
--- a/Engine/Plugins/Experimental/ChaosFlesh/Source/ChaosFleshNodes/Private/Dataflow/ChaosFleshPositionTargetInitializationNodes.cpp
+++ b/Engine/Plugins/Experimental/ChaosFlesh/Source/ChaosFleshNodes/Private/Dataflow/ChaosFleshPositionTargetInitializationNodes.cpp
@@ -14,6 +14,7 @@
 #include "GeometryCollection/Facades/CollectionVertexBoneWeightsFacade.h"
 #include "GeometryCollection/TransformCollection.h"
 #include "Chaos/Tetrahedron.h"
+#include "Chaos/Plane.h"
 
 #include UE_INLINE_GENERATED_CPP_BY_NAME(ChaosFleshPositionTargetInitializationNodes)
 
@@ -26,6 +27,7 @@ namespace Dataflow
 		DATAFLOW_NODE_REGISTER_CREATION_FACTORY(FAddKinematicParticlesDataflowNode);
 		DATAFLOW_NODE_REGISTER_CREATION_FACTORY(FSetVertexVertexPositionTargetBindingDataflowNode);
 		DATAFLOW_NODE_REGISTER_CREATION_FACTORY(FSetVertexTetrahedraPositionTargetBindingDataflowNode);
+		DATAFLOW_NODE_REGISTER_CREATION_FACTORY(FSetVertexTrianglePositionTargetBindingDataflowNode);
 	}
 }
 
@@ -364,4 +366,177 @@ void FSetVertexTetrahedraPositionTargetBindingDataflowNode::Evaluate(Dataflow::F
 	}
 
 
+}
+
+
+void FSetVertexTrianglePositionTargetBindingDataflowNode::Evaluate(Dataflow::FContext& Context, const FDataflowOutput* Out) const
+{
+	if (Out->IsA<DataType>(&Collection))
+	{
+		DataType InCollection = GetValue<DataType>(Context, &Collection);
+
+		if (TManagedArray<FVector3f>* Vertices = InCollection.FindAttribute<FVector3f>("Vertex", FGeometryCollection::VerticesGroup))
+			{
+				if (TManagedArray<FIntVector>* Indices = InCollection.FindAttribute<FIntVector>("Indices", FGeometryCollection::FacesGroup))
+				{
+					if (TManagedArray<int32>* ComponentIndex = InCollection.FindAttribute<int32>("ComponentIndex", FGeometryCollection::VerticesGroup))
+					{
+
+						TArray<Chaos::TVector<int32, 3>> IndicesArray;
+						for (int32 i = 0; i < Indices->Num(); i++)
+						{
+							Chaos::TVector<int32, 3> CurrentIndices(0);
+							for (int32 j = 0; j < 3; j++) 
+							{
+								CurrentIndices[j] = (*Indices)[i][j];
+							}
+							if (CurrentIndices[0] != -1
+								&& CurrentIndices[1] != -1
+								&& CurrentIndices[2] != -1)
+							{
+								IndicesArray.Emplace(CurrentIndices);
+							}
+						}
+						TArray<TArray<int32>> LocalIndex;
+						TArray<TArray<int32>>* LocalIndexPtr = &LocalIndex;
+						TArray<TArray<int>> GlobalIndex = Chaos::Utilities::ComputeIncidentElements(IndicesArray, LocalIndexPtr);
+						int32 ActualParticleCount = 0;
+						for (int32 l = 0; l < GlobalIndex.Num(); l++)
+						{
+							if (GlobalIndex[l].Num() > 0)
+							{
+								ActualParticleCount += 1;
+							}
+						}
+						TArray<Chaos::TVector<float, 3>> IndicesPositions; 
+						IndicesPositions.SetNum(ActualParticleCount);
+						int32 CurrentParticleIndex = 0;
+						for (int32 i = 0; i < GlobalIndex.Num(); i++)
+						{
+							if (GlobalIndex[i].Num() > 0)
+							{
+								IndicesPositions[CurrentParticleIndex] = (*Vertices)[(*Indices)[GlobalIndex[i][0]][LocalIndex[i][0]]];
+								CurrentParticleIndex += 1;
+							}
+						}
+						Chaos::FReal SphereRadius = (Chaos::FReal)0.;
+
+						Chaos::TVec3<float> CoordMaxs(-FLT_MAX);
+						Chaos::TVec3<float> CoordMins(FLT_MAX);
+						for (int32 i = 0; i < IndicesPositions.Num(); i++)
+						{
+							for (int32 j = 0; j < 3; j++)
+							{
+								if (IndicesPositions[i][j] > CoordMaxs[j])
+								{
+									CoordMaxs[j] = IndicesPositions[i][j];
+								}
+								if (IndicesPositions[i][j] < CoordMins[j])
+								{
+									CoordMins[j] = IndicesPositions[i][j];
+								}
+							}
+						}
+						Chaos::TVec3<float> CoordDiff = (CoordMaxs - CoordMins) * VertexRadiusRatio;
+						SphereRadius = Chaos::FReal(FGenericPlatformMath::Min(CoordDiff[0], FGenericPlatformMath::Min(CoordDiff[1], CoordDiff[2])));
+
+						TArray<Chaos::TSphere<Chaos::FReal, 3>*> VertexSpherePtrs;
+						TArray<Chaos::TSphere<Chaos::FReal, 3>> VertexSpheres;
+
+						VertexSpheres.Init(Chaos::TSphere<Chaos::FReal, 3>(Chaos::TVec3<Chaos::FReal>(0), SphereRadius), IndicesPositions.Num());
+						VertexSpherePtrs.SetNum(IndicesPositions.Num());
+
+						for (int32 i = 0; i < IndicesPositions.Num(); i++)
+						{
+							Chaos::TVec3<Chaos::FReal> SphereCenter(IndicesPositions[i]);
+							Chaos::TSphere<Chaos::FReal, 3> VertexSphere(SphereCenter, SphereRadius);
+							VertexSpheres[i] = Chaos::TSphere<Chaos::FReal, 3>(SphereCenter, SphereRadius);
+							VertexSpherePtrs[i] = &VertexSpheres[i];
+						}
+						Chaos::TBoundingVolumeHierarchy<
+							TArray<Chaos::TSphere<Chaos::FReal, 3>*>,
+							TArray<int32>,
+							Chaos::FReal,
+							3> VertexBVH(VertexSpherePtrs);
+
+						GeometryCollection::Facades::FPositionTargetFacade PositionTargets(InCollection);
+						PositionTargets.DefineSchema();
+
+						for (int32 i = 0; i < Indices->Num(); i++)
+						{
+							TArray<int32> TriangleIntersections0 = VertexBVH.FindAllIntersections((*Vertices)[(*Indices)[i][0]]);
+							TArray<int32> TriangleIntersections1 = VertexBVH.FindAllIntersections((*Vertices)[(*Indices)[i][1]]);
+							TArray<int32> TriangleIntersections2 = VertexBVH.FindAllIntersections((*Vertices)[(*Indices)[i][2]]);
+							TriangleIntersections0.Sort();
+							TriangleIntersections1.Sort();
+							TriangleIntersections2.Sort();
+
+							TArray<int32> TriangleIntersections({});
+							for (int32 k = 0; k < TriangleIntersections0.Num(); k++)
+							{
+								if (TriangleIntersections1.Contains(TriangleIntersections0[k]) 
+									&& TriangleIntersections2.Contains(TriangleIntersections0[k]))
+								{
+									TriangleIntersections.Emplace(TriangleIntersections0[k]);
+								}
+							}
+
+							int32 TriangleIndex = (*ComponentIndex)[(*Indices)[i][0]];
+							int32 MinIndex = -1;
+							float MinDis = FLT_MAX;
+							Chaos::TVector<float, 3> ClosestBary(0.f);
+							for (int32 j = 0; j < TriangleIntersections.Num(); j++)
+							{
+								if ((*ComponentIndex)[TriangleIntersections[j]] != TriangleIndex)
+								{
+									Chaos::TVector<float, 3> Bary, TriPos0((*Vertices)[(*Indices)[i][0]]), TriPos1((*Vertices)[(*Indices)[i][1]]), TriPos2((*Vertices)[(*Indices)[i][2]]), ParticlePos((*Vertices)[TriangleIntersections[j]]);
+									Chaos::TVector<Chaos::FRealSingle, 3> ClosestPoint = Chaos::FindClosestPointAndBaryOnTriangle(TriPos0, TriPos1, TriPos2, ParticlePos, Bary);
+									Chaos::FRealSingle CurrentDistance = ((*Vertices)[TriangleIntersections[j]] - ClosestPoint).Size();
+									if (CurrentDistance < MinDis)
+									{
+										MinDis = CurrentDistance;
+										MinIndex = TriangleIntersections[j];
+										ClosestBary = Bary;
+									}
+
+								}
+							}
+							if (MinIndex != -1
+								&& MinIndex != (*Indices)[i][0]
+								&& MinIndex != (*Indices)[i][1]
+								&& MinIndex != (*Indices)[i][2])
+							{
+								GeometryCollection::Facades::FPositionTargetsData DataPackage;
+								DataPackage.TargetIndex.Init(MinIndex, 1);
+								DataPackage.TargetWeights.Init(1.f, 1);
+								DataPackage.SourceWeights.Init(1.f, 3);
+								DataPackage.SourceIndex.Init(-1, 3);
+								DataPackage.SourceIndex[0] = (*Indices)[i][0];
+								DataPackage.SourceIndex[1] = (*Indices)[i][1];
+								DataPackage.SourceIndex[2] = (*Indices)[i][2];
+								DataPackage.SourceWeights[0] = ClosestBary[0];
+								DataPackage.SourceWeights[1] = ClosestBary[1];
+								DataPackage.SourceWeights[2] = ClosestBary[2];
+								if (TManagedArray<float>* Mass = InCollection.FindAttribute<float>("Mass", FGeometryCollection::VerticesGroup))
+								{
+									DataPackage.Stiffness = 0.f;
+									for (int32 k = 0; k < 3; k++)
+									{
+										DataPackage.Stiffness += DataPackage.SourceWeights[k] * PositionTargetStiffness * (*Mass)[DataPackage.SourceIndex[k]];
+									}
+									DataPackage.Stiffness += DataPackage.TargetWeights[0] * PositionTargetStiffness * (*Mass)[DataPackage.TargetIndex[0]];
+									DataPackage.Stiffness /= 2.f;
+								}
+								else
+								{
+									DataPackage.Stiffness = PositionTargetStiffness;
+								}
+								PositionTargets.AddPositionTarget(DataPackage);
+							}
+						}
+					}
+				}
+		}
+		SetValue<DataType>(Context, InCollection, &Collection);
+	}
 }
\ No newline at end of file
diff --git a/Engine/Plugins/Experimental/ChaosFlesh/Source/ChaosFleshNodes/Public/Dataflow/ChaosFleshEngineAssetNodes.h b/Engine/Plugins/Experimental/ChaosFlesh/Source/ChaosFleshNodes/Public/Dataflow/ChaosFleshEngineAssetNodes.h
--- a/Engine/Plugins/Experimental/ChaosFlesh/Source/ChaosFleshNodes/Public/Dataflow/ChaosFleshEngineAssetNodes.h
+++ b/Engine/Plugins/Experimental/ChaosFlesh/Source/ChaosFleshNodes/Public/Dataflow/ChaosFleshEngineAssetNodes.h
@@ -159,6 +159,30 @@ public:
 		const TArray<int32>& Insertion) const;
 };
 
+
+
+USTRUCT(meta = (DataflowFlesh))
+struct FComputeIslandsNode : public FDataflowNode
+{
+	GENERATED_USTRUCT_BODY()
+		DATAFLOW_NODE_DEFINE_INTERNAL(FComputeIslandsNode, "ComputeIslands", "Flesh", "")
+		DATAFLOW_NODE_RENDER_TYPE(FGeometryCollection::StaticType(), "Collection")
+
+public:
+
+	UPROPERTY(meta = (DataflowInput, DataflowOutput, DisplayName = "Collection", DataflowPassthrough = "Collection"))
+		FManagedArrayCollection Collection;
+
+	FComputeIslandsNode(const Dataflow::FNodeParameters& InParam, FGuid InGuid = FGuid::NewGuid())
+		: FDataflowNode(InParam, InGuid)
+	{
+		RegisterInputConnection(&Collection);
+		RegisterOutputConnection(&Collection, &Collection);
+	}
+
+	virtual void Evaluate(Dataflow::FContext& Context, const FDataflowOutput* Out) const override;
+};
+
 namespace Dataflow
 {
 	void RegisterChaosFleshEngineAssetNodes();
diff --git a/Engine/Plugins/Experimental/ChaosFlesh/Source/ChaosFleshNodes/Public/Dataflow/ChaosFleshPositionTargetInitializationNodes.h b/Engine/Plugins/Experimental/ChaosFlesh/Source/ChaosFleshNodes/Public/Dataflow/ChaosFleshPositionTargetInitializationNodes.h
--- a/Engine/Plugins/Experimental/ChaosFlesh/Source/ChaosFleshNodes/Public/Dataflow/ChaosFleshPositionTargetInitializationNodes.h
+++ b/Engine/Plugins/Experimental/ChaosFlesh/Source/ChaosFleshNodes/Public/Dataflow/ChaosFleshPositionTargetInitializationNodes.h
@@ -170,6 +170,41 @@ public:
 	virtual void Evaluate(Dataflow::FContext& Context, const FDataflowOutput* Out) const override;
 };
 
+
+USTRUCT(meta = (DataflowFlesh))
+struct FSetVertexTrianglePositionTargetBindingDataflowNode : public FDataflowNode
+{
+	GENERATED_USTRUCT_BODY()
+		DATAFLOW_NODE_DEFINE_INTERNAL(FSetVertexTrianglePositionTargetBindingDataflowNode, "FSetVertexTrianglePositionTargetBinding", "Flesh", "")
+		DATAFLOW_NODE_RENDER_TYPE(FGeometryCollection::StaticType(), "Collection")
+
+public:
+	typedef FManagedArrayCollection DataType;
+
+
+	UPROPERTY(meta = (DataflowInput, DataflowOutput, DisplayName = "Collection", DataflowPassthrough = "Collection"))
+		FManagedArrayCollection Collection;
+
+	UPROPERTY(EditAnywhere, Category = "Dataflow")
+		float PositionTargetStiffness = 10000.f;
+
+	//UPROPERTY(EditAnywhere, Category = "Dataflow", meta = (ClampMin = "0.0", ClampMax = "1.0", UIMin = "0.0", UIMax = "1.0"))
+	//	float TriangleRadiusPadding = .1f;
+
+	UPROPERTY(EditAnywhere, Category = "Dataflow", meta = (ClampMin = "0.0", ClampMax = "2.0", UIMin = "0.0", UIMax = "2.0"))
+		float VertexRadiusRatio = .001f;
+
+
+	FSetVertexTrianglePositionTargetBindingDataflowNode(const Dataflow::FNodeParameters& InParam, FGuid InGuid = FGuid::NewGuid())
+		: FDataflowNode(InParam, InGuid)
+	{
+		RegisterInputConnection(&Collection);
+		RegisterOutputConnection(&Collection, &Collection);
+	}
+
+	virtual void Evaluate(Dataflow::FContext& Context, const FDataflowOutput* Out) const override;
+};
+
 namespace Dataflow
 {
 	void RegisterChaosFleshPositionTargetInitializationNodes();
diff --git a/Engine/Source/Runtime/Experimental/Chaos/Private/Chaos/Deformable/ChaosDeformableSolver.cpp b/Engine/Source/Runtime/Experimental/Chaos/Private/Chaos/Deformable/ChaosDeformableSolver.cpp
--- a/Engine/Source/Runtime/Experimental/Chaos/Private/Chaos/Deformable/ChaosDeformableSolver.cpp
+++ b/Engine/Source/Runtime/Experimental/Chaos/Private/Chaos/Deformable/ChaosDeformableSolver.cpp
@@ -74,6 +74,11 @@ namespace Chaos::Softs
 	FAutoConsoleVariableRef CVarDeformableXPBDCorotatedBatchSize(TEXT("p.Chaos.Deformable.XPBDBatchSize"), GDeformableXPBDCorotatedParams.XPBDCorotatedBatchSize, TEXT("Batch size for physics parallel for. [def: 5]"));
 	FAutoConsoleVariableRef CVarDeformableXPBDCorotatedBatchThreshold(TEXT("p.Chaos.Deformable.XPBDBatchThreshold"), GDeformableXPBDCorotatedParams.XPBDCorotatedBatchThreshold, TEXT("Batch threshold for physics parallel for. [def: 5]"));
 
+	FDeformableXPBDWeakConstraintParams GDeformableXPBDWeakConstraintParams;
+	FAutoConsoleVariableRef CVarDeformableXPBDWeakConstraintLineWidth(TEXT("p.Chaos.Deformable.XPBDWeakConstraintLineWidth"), GDeformableXPBDWeakConstraintParams.DebugLineWidth, TEXT("Line width for visualizing the double bindings in XPBD weak constraints. [def: 5]"));
+	FAutoConsoleVariableRef CVarDeformableXPBDWeakConstraintParticleWidth(TEXT("p.Chaos.Deformable.XPBDWeakConstraintParticleWidth"), GDeformableXPBDWeakConstraintParams.DebugParticleWidth, TEXT("Line width for visualizing the double bindings in XPBD weak constraints. [def: 20]"));
+	FAutoConsoleVariableRef CVarDeformableXPBDWeakConstraintDebugDraw(TEXT("p.Chaos.Deformable.XPBDWeakConstraintEnableDraw"), GDeformableXPBDWeakConstraintParams.bVisualizeBindings, TEXT("Debug draw the double bindings in XPBD weak constraints. [def: false]"));
+
 	FCriticalSection FDeformableSolver::InitializationMutex;
 	FCriticalSection FDeformableSolver::RemovalMutex;
 	FCriticalSection FDeformableSolver::PackageOutputMutex;
@@ -456,13 +461,12 @@ namespace Chaos::Softs
 
 			FXPBDWeakConstraints<FSolverReal, FSolverParticles>* WeakConstraint =
 				new FXPBDWeakConstraints<FSolverReal, FSolverParticles>(Evolution->Particles(),
-					PositionTargetIndices, PositionTargetWeights, PositionTargetStiffness, PositionTargetSecondIndices, PositionTargetSecondWeights);
+					PositionTargetIndices, PositionTargetWeights, PositionTargetStiffness, PositionTargetSecondIndices, PositionTargetSecondWeights, GDeformableXPBDWeakConstraintParams);
 
 			Evolution->ConstraintInits()[InitIndex] =
 				[WeakConstraint, this](FSolverParticles& InParticles, const FSolverReal Dt)
 			{
-				WeakConstraint->Init();
-				//WeakConstraint->UpdateTargets(this->ComputeParticleTargets(WeakConstraint->GetIndices()));
+				WeakConstraint->Init(InParticles, Dt);
 			};
 
 			Evolution->ConstraintRules()[ConstraintIndex] =
diff --git a/Engine/Source/Runtime/Experimental/Chaos/Public/Chaos/Utilities.h b/Engine/Source/Runtime/Experimental/Chaos/Public/Chaos/Utilities.h
--- a/Engine/Source/Runtime/Experimental/Chaos/Public/Chaos/Utilities.h
+++ b/Engine/Source/Runtime/Experimental/Chaos/Public/Chaos/Utilities.h
@@ -910,5 +910,91 @@ namespace Chaos
 			return IncidentElements;
 		}
 
+		inline void DFS_iterative(const TArray<TArray<int32>>& L, int32 v, TArray<bool>& visited, TArray<int32>& component) {
+			TArray<int32> Stack({ v });
+			Stack.Heapify();
+			while (!Stack.IsEmpty()) {
+				Stack.HeapPop(v);
+				if (!visited[v]) {
+					visited[v] = true;
+					component.Emplace(v);
+					for (int32 i : L[v]) {
+						if (!visited[i]) {
+							Stack.HeapPush(i);
+						}
+					}
+				}
+			}
+		}
+
+		inline void ConnectedComponentsDFSIterative(const TArray<TArray<int32>>& L, TArray<TArray<int32>>& C) {
+			//Input: L, the given adjacency list
+			//Output: C, connected components
+			TArray<bool> Visited;
+			Visited.Init(false, L.Num());
+			for (int32 v = 0; v < L.Num(); ++v) {
+				if (!Visited[v]) {
+					TArray<int32> component;
+					DFS_iterative(L, v, Visited, component);
+					C.Emplace(component);
+				}
+			}
+		}
+
+		inline void FindConnectedRegions(const TArray<FIntVector4>& Elements, TArray<TArray<int32>>& ConnectedComponents) 
+		{
+			TArray<int32> AllEntries, ElementIndices, Ordering, Ranges;
+			int32 count = 0;
+			AllEntries.Init(-1, Elements.Num() * 4);
+			ElementIndices.Init(-1, Elements.Num()*4);
+			Ordering.Init(-1, Elements.Num()*4);
+			Ranges.Init(-1, Elements.Num()*4 + 1);
+			for (int32 i = 0; i < Elements.Num(); i++) 
+			{
+				for (int32 j = 0; j < 4; j++) 
+				{
+					AllEntries[4 * i + j] = Elements[i][j];
+					ElementIndices[4*i+j] = i;
+					Ordering[4*i+j] = count;
+					Ranges[4*i+j] = count++;
+				}
+			}
+			Ranges[Elements.Num() * 4] = count;
+			
+			if (AllEntries.Num() == 0)
+				return;
+			
+			Ordering.Sort([&AllEntries](int32 a, int32 b) { return AllEntries[a] < AllEntries[b]; });
+			TArray<int32> UniqueRanges({ 0 });
+			
+			for (int32 i = 0; i < Ranges.Num() - 2; ++i)
+			{
+				if (AllEntries[Ordering[i]] != AllEntries[Ordering[i+1]])
+				{
+					UniqueRanges.Emplace(i + 1);
+				}
+			}
+			
+			UniqueRanges.Emplace(count);
+
+			TArray<TArray<int32>> AdjacencyList;
+			AdjacencyList.Init(TArray<int32>(), Elements.Num());
+			for (int32 r = 0; r < UniqueRanges.Num() - 1; r++) 
+			{
+				if (UniqueRanges[r + 1] - UniqueRanges[r] > 1) 
+				{
+					for (int32 i = UniqueRanges[r]; i < UniqueRanges[r + 1]; i++) 
+					{
+						for (int32 j = i + 1; j < UniqueRanges[r + 1]; j++) 
+						{
+							AdjacencyList[ElementIndices[Ordering[i]]].Emplace(ElementIndices[Ordering[j]]);
+							AdjacencyList[ElementIndices[Ordering[j]]].Emplace(ElementIndices[Ordering[i]]);
+						}
+					}
+				}
+			}
+			ConnectedComponentsDFSIterative(AdjacencyList, ConnectedComponents);
+		}
+
 	} // namespace Utilities
 } // namespace Chaos
\ No newline at end of file
diff --git a/Engine/Source/Runtime/Experimental/Chaos/Public/Chaos/XPBDWeakConstraints.h b/Engine/Source/Runtime/Experimental/Chaos/Public/Chaos/XPBDWeakConstraints.h
--- a/Engine/Source/Runtime/Experimental/Chaos/Public/Chaos/XPBDWeakConstraints.h
+++ b/Engine/Source/Runtime/Experimental/Chaos/Public/Chaos/XPBDWeakConstraints.h
@@ -12,6 +12,7 @@
 #include "HAL/IConsoleManager.h"
 #include "Misc/Paths.h"
 #include "Chaos/GraphColoring.h"
+#include "Chaos/DebugDrawQueue.h"
 
 //DECLARE_CYCLE_STAT(TEXT("Chaos XPBD Corotated Constraint"), STAT_ChaosXPBDCorotated, STATGROUP_Chaos);
 //DECLARE_CYCLE_STAT(TEXT("Chaos XPBD Corotated Constraint Polar Compute"), STAT_ChaosXPBDCorotatedPolar, STATGROUP_Chaos);
@@ -20,6 +21,13 @@
 namespace Chaos::Softs
 {
 
+	struct CHAOS_API FDeformableXPBDWeakConstraintParams
+	{
+		float DebugLineWidth = 5.f;
+		float DebugParticleWidth = 20.f;
+		bool bVisualizeBindings = false;
+	};
+
 	using Chaos::TVec3;
 
 	template <typename T, typename ParticleType>
@@ -48,9 +56,10 @@ namespace Chaos::Softs
 			const TArray<TArray<T>>& InWeights,
 			const TArray<T>& InStiffness,
 			const TArray<TArray<int32>>& InSecondIndices,
-			const TArray<TArray<T>>& InSecondWeights
+			const TArray<TArray<T>>& InSecondWeights,
+			const FDeformableXPBDWeakConstraintParams& InParams
 		)
-			: Indices(InIndices), Weights(InWeights), SecondIndices(InSecondIndices), SecondWeights(InSecondWeights), Stiffness(InStiffness)
+			: Indices(InIndices), Weights(InWeights), SecondIndices(InSecondIndices), SecondWeights(InSecondWeights), Stiffness(InStiffness), DebugDrawParams(InParams)
 		{
 			ensureMsgf(Indices.Num() == SecondIndices.Num(), TEXT("Input Double Bindings have wrong size"));
 			
@@ -106,9 +115,13 @@ namespace Chaos::Softs
 			return Indices;
 		}
 
-		virtual void Init() const
+		void Init(const FSolverParticles& InParticles, const T Dt) const
 		{
 			for (T& Lambdas : LambdaArray) { Lambdas = (T)0.; }
+			if (DebugDrawParams.bVisualizeBindings)
+			{
+				VisualizeAllBindings(InParticles, Dt);
+			}
 		}
 
 		void UpdateTargets(TArray<TVector<T, 3>>&& InTargets)
@@ -222,6 +235,51 @@ namespace Chaos::Softs
 			}
 		}
 
+		void VisualizeAllBindings(const FSolverParticles& InParticles, const T Dt) const
+		{
+#if WITH_EDITOR
+			auto DoubleVert = [](Chaos::TVec3<T> V) { return FVector3d(V.X, V.Y, V.Z); };
+			for (int32 i = 0; i < Indices.Num(); i++)
+			{
+				Chaos::TVec3<T> SourcePos((T)0.), TargetPos((T)0.);
+				for (int32 j = 0; j < Indices[i].Num(); j++)
+				{
+					SourcePos += Weights[i][j] * InParticles.P(Indices[i][j]);
+				}
+				for (int32 j = 0; j < SecondIndices[i].Num(); j++)
+				{
+					TargetPos += SecondWeights[i][j] * InParticles.P(SecondIndices[i][j]);
+				}
+
+				float ParticleThickness = DebugDrawParams.DebugParticleWidth;
+				float LineThickness = DebugDrawParams.DebugLineWidth;
+
+				if (Indices[i].Num() == 1)
+				{
+					Chaos::FDebugDrawQueue::GetInstance().DrawDebugPoint(DoubleVert(SourcePos), FColor::Red, false, Dt, 0, ParticleThickness);
+					for (int32 j = 0; j < SecondIndices[i].Num(); j++)
+					{
+						Chaos::FDebugDrawQueue::GetInstance().DrawDebugPoint(DoubleVert(InParticles.P(SecondIndices[i][j])), FColor::Green, false, Dt, 0, ParticleThickness);
+						Chaos::FDebugDrawQueue::GetInstance().DrawDebugLine(DoubleVert(InParticles.P(SecondIndices[i][j])), DoubleVert(InParticles.P(SecondIndices[i][(j+1) % SecondIndices[i].Num()])), FColor::Green, false, Dt, 0, LineThickness);
+					}
+				
+				}
+
+				if (SecondIndices[i].Num() == 1)
+				{
+					Chaos::FDebugDrawQueue::GetInstance().DrawDebugPoint(DoubleVert(TargetPos), FColor::Red, false, Dt, 0, ParticleThickness);
+					for (int32 j = 0; j < Indices[i].Num(); j++)
+					{
+						Chaos::FDebugDrawQueue::GetInstance().DrawDebugPoint(DoubleVert(InParticles.P(Indices[i][j])), FColor::Green, false, Dt, 0, ParticleThickness);
+						Chaos::FDebugDrawQueue::GetInstance().DrawDebugLine(DoubleVert(InParticles.P(Indices[i][j])), DoubleVert(InParticles.P(Indices[i][(j + 1) % SecondIndices[i].Num()])), FColor::Green, false, Dt, 0, LineThickness);
+					}
+				}
+
+				Chaos::FDebugDrawQueue::GetInstance().DrawDebugLine(DoubleVert(SourcePos), DoubleVert(TargetPos), FColor::Yellow, false, Dt, 0, LineThickness);
+			}
+#endif
+		}
+
 
 	protected:
 		TArray<TArray<int32>> Indices;
@@ -232,7 +290,11 @@ namespace Chaos::Softs
 		TArray<T> Stiffness;
 		TArray<TArray<int32>> ConstraintsPerColor;
 		mutable TArray<T> LambdaArray;
+		bool VisualizeBindings = false;
+		FDeformableXPBDWeakConstraintParams DebugDrawParams;
 	};
 
+	
+
 }  // End namespace Chaos::Softs
 
