ChaosFlesh : Add a flesh asset input that injects a separate flesh asset into the graph.

diff --git a/Engine/Plugins/Experimental/ChaosFlesh/Source/ChaosFleshNodes/Private/Dataflow/ChaosFleshEngineAssetNodes.cpp b/Engine/Plugins/Experimental/ChaosFlesh/Source/ChaosFleshNodes/Private/Dataflow/ChaosFleshEngineAssetNodes.cpp
--- a/Engine/Plugins/Experimental/ChaosFlesh/Source/ChaosFleshNodes/Private/Dataflow/ChaosFleshEngineAssetNodes.cpp
+++ b/Engine/Plugins/Experimental/ChaosFlesh/Source/ChaosFleshNodes/Private/Dataflow/ChaosFleshEngineAssetNodes.cpp
@@ -31,14 +31,23 @@ void FGetFleshAssetDataflowNode::Evaluate(Dataflow::FContext& Context, const FDa
 {
 	if (Out->IsA<FManagedArrayCollection>(&Output))
 	{
-		if (const Dataflow::FEngineContext* EngineContext = Context.AsType<Dataflow::FEngineContext>())
+		FManagedArrayCollection Collection;
+		SetValue<FManagedArrayCollection>(Context, Collection, &Output);
+
+		const UFleshAsset* FleshAssetValue = FleshAsset;
+		if (!FleshAssetValue)
 		{
-			if (UFleshAsset* FleshAsset = Cast<UFleshAsset>(EngineContext->Owner))
+			if (const Dataflow::FEngineContext* EngineContext = Context.AsType<Dataflow::FEngineContext>())
 			{
-				if (const FFleshCollection* AssetCollection = FleshAsset->GetCollection())
-				{
-					SetValue<FManagedArrayCollection>(Context, *AssetCollection, &Output);
-				}
+				FleshAssetValue = Cast<UFleshAsset>(EngineContext->Owner);
+			}
+		}
+
+		if (FleshAssetValue)
+		{
+			if (const FFleshCollection* AssetCollection = FleshAssetValue->GetCollection())
+			{
+				SetValue<FManagedArrayCollection>(Context, *AssetCollection, &Output);
 			}
 		}
 	}
diff --git a/Engine/Plugins/Experimental/ChaosFlesh/Source/ChaosFleshNodes/Public/Dataflow/ChaosFleshEngineAssetNodes.h b/Engine/Plugins/Experimental/ChaosFlesh/Source/ChaosFleshNodes/Public/Dataflow/ChaosFleshEngineAssetNodes.h
--- a/Engine/Plugins/Experimental/ChaosFlesh/Source/ChaosFleshNodes/Public/Dataflow/ChaosFleshEngineAssetNodes.h
+++ b/Engine/Plugins/Experimental/ChaosFlesh/Source/ChaosFleshNodes/Public/Dataflow/ChaosFleshEngineAssetNodes.h
@@ -21,6 +21,9 @@ struct FGetFleshAssetDataflowNode : public FDataflowNode
 
 public:
 
+	UPROPERTY(EditAnywhere, Category = "Dataflow", meta = (DisplayName = "FleshAsset"))
+	TObjectPtr<const UFleshAsset> FleshAsset = nullptr;
+
 	UPROPERTY(meta = (DataflowOutput, DisplayName = "Collection"))
 	FManagedArrayCollection Output;
 
