ChaosFlesh: surface bindings and weak constraint driven flesh sim

diff --git a/Engine/Plugins/Experimental/ChaosFlesh/Source/ChaosFleshEngine/Private/ChaosFlesh/FleshComponent.cpp b/Engine/Plugins/Experimental/ChaosFlesh/Source/ChaosFleshEngine/Private/ChaosFlesh/FleshComponent.cpp
--- a/Engine/Plugins/Experimental/ChaosFlesh/Source/ChaosFleshEngine/Private/ChaosFlesh/FleshComponent.cpp
+++ b/Engine/Plugins/Experimental/ChaosFlesh/Source/ChaosFleshEngine/Private/ChaosFlesh/FleshComponent.cpp
@@ -94,25 +94,28 @@ UDeformablePhysicsComponent::FThreadingProxy* UFleshComponent::NewProxy()
 	{
 		if (const FFleshCollection* Rest = RestAsset->GetCollection())
 		{
-			if (!GetDynamicCollection())
+			if (Rest->NumElements(FGeometryCollection::VerticesGroup))
 			{
-				DynamicCollection = NewObject<UFleshDynamicAsset>(this, TEXT("Flesh Dynamic Asset"));
-			}
+				if (!GetDynamicCollection())
+				{
+					DynamicCollection = NewObject<UFleshDynamicAsset>(this, TEXT("Flesh Dynamic Asset"));
+				}
 
-			GetDynamicCollection()->Reset(Rest);
-			if (const FManagedArrayCollection* Dynamic = GetDynamicCollection()->GetCollection())
-			{
-				// Mesh points are in component space, such that the exterior hull aligns with the
-                // surface of the skeletal mesh, which is subject to the transform hierarchy.
-				const FTransform& ComponentToWorldXf = GetComponentTransform();
-				const FTransform ComponentToSimXf = GetSimSpaceRestTransform();
-				return new FFleshThreadingProxy(
-					this,
-					ComponentToWorldXf, 
-					ComponentToSimXf, 
-					SimSpace,
-					*Rest, 
-					*Dynamic);
+				GetDynamicCollection()->Reset(Rest);
+				if (const FManagedArrayCollection* Dynamic = GetDynamicCollection()->GetCollection())
+				{
+					// Mesh points are in component space, such that the exterior hull aligns with the
+					// surface of the skeletal mesh, which is subject to the transform hierarchy.
+					const FTransform& ComponentToWorldXf = GetComponentTransform();
+					const FTransform ComponentToSimXf = GetSimSpaceRestTransform();
+					return new FFleshThreadingProxy(
+						this,
+						ComponentToWorldXf,
+						ComponentToSimXf,
+						SimSpace,
+						*Rest,
+						*Dynamic);
+				}
 			}
 		}
 	}
diff --git a/Engine/Plugins/Experimental/ChaosFlesh/Source/ChaosFleshNodes/Private/Dataflow/ChaosFleshKinematicInitializationNodes.cpp b/Engine/Plugins/Experimental/ChaosFlesh/Source/ChaosFleshNodes/Private/Dataflow/ChaosFleshKinematicInitializationNodes.cpp
--- a/Engine/Plugins/Experimental/ChaosFlesh/Source/ChaosFleshNodes/Private/Dataflow/ChaosFleshKinematicInitializationNodes.cpp
+++ b/Engine/Plugins/Experimental/ChaosFlesh/Source/ChaosFleshNodes/Private/Dataflow/ChaosFleshKinematicInitializationNodes.cpp
@@ -12,6 +12,8 @@
 #include "GeometryCollection/Facades/CollectionVertexBoneWeightsFacade.h"
 #include "GeometryCollection/TransformCollection.h"
 #include "PhysicsEngine/PhysicsAsset.h"
+#include "Rendering/SkeletalMeshRenderData.h"
+#include "BoneWeights.h"
 
 #include UE_INLINE_GENERATED_CPP_BY_NAME(ChaosFleshKinematicInitializationNodes)
 
@@ -21,8 +23,10 @@ namespace Dataflow
 {
 	void RegisterChaosFleshKinematicInitializationNodes()
 	{
+		DATAFLOW_NODE_REGISTER_CREATION_FACTORY(FKinematicSkeletalMeshInitializationDataflowNode);
 		DATAFLOW_NODE_REGISTER_CREATION_FACTORY(FKinematicBodySetupInitializationDataflowNode);
 		DATAFLOW_NODE_REGISTER_CREATION_FACTORY(FKinematicInitializationDataflowNode);
+		DATAFLOW_NODE_REGISTER_CREATION_FACTORY(FKinematicOriginInsertionInitializationDataflowNode);
 		DATAFLOW_NODE_REGISTER_CREATION_FACTORY(FKinematicTetrahedralBindingsDataflowNode);
 		DATAFLOW_NODE_REGISTER_CREATION_FACTORY(FSetVerticesKinematicDataflowNode);
 	}
@@ -42,11 +46,11 @@ void FKinematicTetrahedralBindingsDataflowNode::Evaluate(Dataflow::FContext& Con
 			//parse exclusion list to find bones to skip
 			TArray<FString> StrArray;
 			ExclusionList.ParseIntoArray(StrArray, *FString(" "));				
-
+			TSet<int32> AddedVertSet;
 			int32 NumTets = Tetrahedron->Num();
 			TArray<FTransform> ComponentPose;
 			Dataflow::Animation::GlobalTransforms(SkeletalMesh->GetRefSkeleton(), ComponentPose);
-			for (int32 b = 0; b < SkeletalMesh->GetRefSkeleton().GetNum(); b++)
+			for (int32 b = SkeletalMesh->GetRefSkeleton().GetNum()-1; b > -1; b--)
 			{
 				bool Skip = false;
 				FString BoneName = SkeletalMesh->GetRefSkeleton().GetBoneName(b).ToString();
@@ -93,31 +97,40 @@ void FKinematicTetrahedralBindingsDataflowNode::Evaluate(Dataflow::FContext& Con
 							int32 OutFaceIndex;
 							bool KeepTet = ConvexTet.Raycast(BonePosition, RayDir, Length, Chaos::FReal(0), OutTime, OutPosition, OutNormal, OutFaceIndex);
 							if (KeepTet) 
-							{
-								BoneVertSet.Add(i); BoneVertSet.Add(j); BoneVertSet.Add(k); BoneVertSet.Add(l);
+							{	
+								for (int32 c = 0; c < 4; ++c)
+								{
+									if (!AddedVertSet.Contains((*Tetrahedron)[t][c]))
+									{
+										AddedVertSet.Add((*Tetrahedron)[t][c]);
+										BoneVertSet.Add((*Tetrahedron)[t][c]);
+									}
+								}
 							}
 						}
 
 						TArray<int32> BoundVerts = BoneVertSet.Array();
 						TArray<float> BoundWeights;
 						BoundWeights.Init(float(1), BoundVerts.Num());
-						
-						//get local coords of bound verts
-						typedef GeometryCollection::Facades::FKinematicBindingFacade FKinematics;
-						FKinematics Kinematics(InCollection); Kinematics.DefineSchema();
-						if (Kinematics.IsValid())
+						if (BoundVerts.Num())
 						{
-							FKinematics::FBindingKey Binding = Kinematics.SetBoneBindings(b, BoundVerts, BoundWeights);
-							TManagedArray<TArray<FVector3f>>& LocalPos = InCollection.AddAttribute<TArray<FVector3f>>("LocalPosition", Binding.GroupName);
-							Kinematics.AddKinematicBinding(Binding);
-
-							auto FloatVert = [](FVector3d V) { return FVector3f(V.X, V.Y, V.Z); };
-							auto DoubleVert = [](FVector3f V) { return FVector3d(V.X, V.Y, V.Z); };
-							LocalPos[Binding.Index].SetNum(BoundVerts.Num());
-							for (int32 i = 0; i < BoundVerts.Num(); i++)
+							//get local coords of bound verts
+							typedef GeometryCollection::Facades::FKinematicBindingFacade FKinematics;
+							FKinematics Kinematics(InCollection); Kinematics.DefineSchema();
+							if (Kinematics.IsValid())
 							{
-								FVector3f Temp = (*Vertex)[BoundVerts[i]];
-								LocalPos[Binding.Index][i] = FloatVert(ComponentPose[b].InverseTransformPosition(DoubleVert(Temp)));
+								FKinematics::FBindingKey Binding = Kinematics.SetBoneBindings(b, BoundVerts, BoundWeights);
+								TManagedArray<TArray<FVector3f>>& LocalPos = InCollection.AddAttribute<TArray<FVector3f>>("LocalPosition", Binding.GroupName);
+								Kinematics.AddKinematicBinding(Binding);
+
+								auto FloatVert = [](FVector3d V) { return FVector3f(V.X, V.Y, V.Z); };
+								auto DoubleVert = [](FVector3f V) { return FVector3d(V.X, V.Y, V.Z); };
+								LocalPos[Binding.Index].SetNum(BoundVerts.Num());
+								for (int32 i = 0; i < BoundVerts.Num(); i++)
+								{
+									FVector3f Temp = (*Vertex)[BoundVerts[i]];
+									LocalPos[Binding.Index][i] = FloatVert(ComponentPose[b].InverseTransformPosition(DoubleVert(Temp)));
+								}
 							}
 						}
 					}
@@ -201,7 +214,7 @@ void FKinematicInitializationDataflowNode::Evaluate(Dataflow::FContext& Context,
 						TArray<int32> BoundVerts;
 						TArray<float> BoundWeights;
 
-						if (0 < Index && Index < ComponentPose.Num())
+						if (0 <= Index && Index < ComponentPose.Num())
 						{
 							FVector3f BonePosition(ComponentPose[Index].GetTranslation());
 
@@ -229,7 +242,131 @@ void FKinematicInitializationDataflowNode::Evaluate(Dataflow::FContext& Context,
 					GeometryCollection::Facades::FVertexBoneWeightsFacade(InCollection).AddBoneWeightsFromKinematicBindings();
 				}
 			}
+		}
+		SetValue<DataType>(Context, InCollection, &Collection);
+	}
+}
+
+void FKinematicOriginInsertionInitializationDataflowNode::Evaluate(Dataflow::FContext& Context, const FDataflowOutput* Out) const
+{
+	if (Out->IsA<DataType>(&Collection))
+	{
+		DataType InCollection = GetValue<DataType>(Context, &Collection);
 
+		if (TManagedArray<FVector3f>* Vertices = InCollection.FindAttribute<FVector3f>("Vertex", FGeometryCollection::VerticesGroup))
+		{
+			if (FindInput(&OriginVertexIndicesIn) && FindInput(&OriginVertexIndicesIn)->GetConnection() && FindInput(&InsertionVertexIndicesIn) && FindInput(&InsertionVertexIndicesIn)->GetConnection())
+			{
+				TArray<int32> BoundVerts;
+				TArray<float> BoundWeights;
+
+				for (int32 SelectionIndex : GetValue<TArray<int32>>(Context, &OriginVertexIndicesIn))
+				{
+					if (0 <= SelectionIndex && SelectionIndex < Vertices->Num())
+					{
+						BoundVerts.Add(SelectionIndex);
+					}
+				}
+				for (int32 SelectionIndex : GetValue<TArray<int32>>(Context, &InsertionVertexIndicesIn))
+				{
+					if (0 <= SelectionIndex && SelectionIndex < Vertices->Num())
+					{
+						BoundVerts.Add(SelectionIndex);
+					}
+				}
+				if (TObjectPtr<USkeletalMesh> BoneSkeletalMesh = GetValue<TObjectPtr<USkeletalMesh>>(Context, &BoneSkeletalMeshIn))
+				{
+					FSkeletalMeshRenderData* RenderData = BoneSkeletalMesh->GetResourceForRendering();
+					if (RenderData->LODRenderData.Num())
+					{
+						//Grab vertices only, no elements
+						FSkeletalMeshLODRenderData* LODRenderData = &RenderData->LODRenderData[0];
+						const FPositionVertexBuffer& PositionVertexBuffer =
+							LODRenderData->StaticVertexBuffers.PositionVertexBuffer;
+						//Grab skin weights
+						const FSkinWeightVertexBuffer* SkinWeightVertexBuffer = LODRenderData->GetSkinWeightVertexBuffer();
+						const int32 MaxBoneInfluences = SkinWeightVertexBuffer->GetMaxBoneInfluences();
+						TArray<FTransform> ComponentPose;
+						Dataflow::Animation::GlobalTransforms(BoneSkeletalMesh->GetRefSkeleton(), ComponentPose);
+						TArray<TArray<int32>> BoneBoundVerts;
+						TArray<TArray<float>> BoneBoundWeights;
+						BoneBoundVerts.SetNum(ComponentPose.Num());
+						BoneBoundWeights.SetNum(ComponentPose.Num());
+						if (!PositionVertexBuffer.GetNumVertices())
+						{
+							return;
+						}
+						auto DoubleVert = [](FVector3f V) { return FVector3d(V.X, V.Y, V.Z); };
+						for (int32 i = 0; i < BoundVerts.Num(); i++)
+						{
+							int ClosestPointIndex = 0;
+							double MinDistance = DBL_MAX;
+							for (uint32 j = 0; j < PositionVertexBuffer.GetNumVertices(); j++)
+							{
+								const FVector3f& Pos = PositionVertexBuffer.VertexPosition(j);
+								double Distance = FVector::Distance(DoubleVert((*Vertices)[BoundVerts[i]]), DoubleVert(Pos));
+								if (Distance < MinDistance)
+								{
+									ClosestPointIndex = j;
+									MinDistance = Distance;
+								}
+							}
+							int32 SectionIndex;
+							int32 VertIndex;
+							LODRenderData->GetSectionFromVertexIndex(ClosestPointIndex, SectionIndex, VertIndex);
+
+							check(SectionIndex < LODRenderData->RenderSections.Num());
+							const FSkelMeshRenderSection& Section = LODRenderData->RenderSections[SectionIndex];
+							int32 BufferVertIndex = Section.GetVertexBufferIndex() + VertIndex;
+							for (int32 InfluenceIndex = 0; InfluenceIndex < MaxBoneInfluences; InfluenceIndex++)
+							{
+								const int32 BoneIndex = Section.BoneMap[SkinWeightVertexBuffer->GetBoneIndex(BufferVertIndex, InfluenceIndex)];
+								const float	Weight = (float)SkinWeightVertexBuffer->GetBoneWeight(BufferVertIndex, InfluenceIndex) * UE::AnimationCore::InvMaxRawBoneWeightFloat;
+								if (Weight > float(0) && 0 <= BoneIndex && BoneIndex < ComponentPose.Num())
+								{
+									FString BoneName = BoneSkeletalMesh->GetRefSkeleton().GetBoneName(BoneIndex).ToString();
+									BoneBoundVerts[BoneIndex].Add(BoundVerts[i]);
+									BoneBoundWeights[BoneIndex].Add(Weight);
+								}
+							}
+						}
+						for (int32 BoneIndex = 0; BoneIndex < ComponentPose.Num(); ++BoneIndex)
+						{
+							if (BoneBoundVerts[BoneIndex].Num())
+							{
+								FString BoneName = BoneSkeletalMesh->GetRefSkeleton().GetBoneName(BoneIndex).ToString();
+								//get local coords of bound verts
+								typedef GeometryCollection::Facades::FKinematicBindingFacade FKinematics;
+								FKinematics Kinematics(InCollection); Kinematics.DefineSchema();
+								if (Kinematics.IsValid())
+								{
+									FKinematics::FBindingKey Binding = Kinematics.SetBoneBindings(BoneIndex, BoneBoundVerts[BoneIndex], BoneBoundWeights[BoneIndex]);
+									TManagedArray<TArray<FVector3f>>& LocalPos = InCollection.AddAttribute<TArray<FVector3f>>("LocalPosition", Binding.GroupName);
+									Kinematics.AddKinematicBinding(Binding);
+
+									auto FloatVert = [](FVector3d V) { return FVector3f(V.X, V.Y, V.Z); };
+									LocalPos[Binding.Index].SetNum(BoneBoundVerts[BoneIndex].Num());
+									for (int32 i = 0; i < BoneBoundVerts[BoneIndex].Num(); i++)
+									{
+										FVector3f Temp = (*Vertices)[BoneBoundVerts[BoneIndex][i]];
+										LocalPos[Binding.Index][i] = FloatVert(ComponentPose[BoneIndex].InverseTransformPosition(DoubleVert(Temp)));
+									}
+								}
+							}
+						}
+						GeometryCollection::Facades::FVertexBoneWeightsFacade(InCollection).AddBoneWeightsFromKinematicBindings();
+					}
+				}
+				else
+				{
+					if (BoundVerts.Num())
+					{
+						BoundWeights.Init(1.0, BoundVerts.Num());
+						GeometryCollection::Facades::FKinematicBindingFacade Kinematics(InCollection);
+						Kinematics.AddKinematicBinding(Kinematics.SetBoneBindings(INDEX_NONE, BoundVerts, BoundWeights));
+					}
+				}
+			}
 		}
 		SetValue<DataType>(Context, InCollection, &Collection);
 	}
@@ -307,7 +444,7 @@ void FKinematicBodySetupInitializationDataflowNode::Evaluate(Dataflow::FContext&
 					{	
 						TArray<FKSphylElem> SphylElems = BodySetup->AggGeom.SphylElems;
 						int32 BoneIndex = SkeletalMesh->GetRefSkeleton().FindBoneIndex(BodySetup->BoneName);
-						if (0 < BoneIndex && BoneIndex < ComponentPose.Num())
+						if (0 <= BoneIndex && BoneIndex < ComponentPose.Num())
 						{
 							TArray<int32> BoundVerts;
 							TArray<float> BoundWeights;
@@ -353,4 +490,91 @@ void FKinematicBodySetupInitializationDataflowNode::Evaluate(Dataflow::FContext&
 		}
 		SetValue<DataType>(Context, InCollection, &Collection);
 	}
+}
+
+void FKinematicSkeletalMeshInitializationDataflowNode::Evaluate(Dataflow::FContext& Context, const FDataflowOutput* Out) const
+{
+	if (Out->IsA<DataType>(&Collection) || Out->IsA<TArray<int32>>(&IndicesOut))
+	{
+		DataType InCollection = GetValue<DataType>(Context, &Collection);
+		TArray<int32> Indices;
+		if (TObjectPtr<USkeletalMesh> SkeletalMesh = GetValue<TObjectPtr<USkeletalMesh>>(Context, &SkeletalMeshIn))
+		{
+			FSkeletalMeshRenderData* RenderData = SkeletalMesh->GetResourceForRendering();
+			if (RenderData->LODRenderData.Num())
+			{
+				//Grab vertices only, no elements
+				FSkeletalMeshLODRenderData* LODRenderData = &RenderData->LODRenderData[0];
+				const FPositionVertexBuffer& PositionVertexBuffer =
+					LODRenderData->StaticVertexBuffers.PositionVertexBuffer;
+				TManagedArray<FVector3f>& Vertices = InCollection.AddAttribute<FVector3f>("Vertex", FGeometryCollection::VerticesGroup);
+				int32 index = InCollection.AddElements(PositionVertexBuffer.GetNumVertices(), FGeometryCollection::VerticesGroup);
+				for (uint32 j = 0; j < PositionVertexBuffer.GetNumVertices(); j++)
+				{
+					const FVector3f& Pos = PositionVertexBuffer.VertexPosition(j);
+					Vertices[index+j] = Pos;
+					Indices.Add(index+j);
+				}
+				//Grab skin weights
+				const FSkinWeightVertexBuffer* SkinWeightVertexBuffer = LODRenderData->GetSkinWeightVertexBuffer();
+				const int32 MaxBoneInfluences = SkinWeightVertexBuffer->GetMaxBoneInfluences();
+				TArray<FTransform> ComponentPose;
+				Dataflow::Animation::GlobalTransforms(SkeletalMesh->GetRefSkeleton(), ComponentPose);
+				TArray<TArray<int32>> BoundVerts;
+				TArray<TArray<float>> BoundWeights;
+				BoundVerts.SetNum(ComponentPose.Num());
+				BoundWeights.SetNum(ComponentPose.Num());
+				for (uint32 j = 0; j < PositionVertexBuffer.GetNumVertices(); j++)
+				{	
+					
+					int32 SectionIndex;
+					int32 VertIndex;
+					LODRenderData->GetSectionFromVertexIndex(j, SectionIndex, VertIndex);
+
+					check(SectionIndex < LODRenderData->RenderSections.Num());
+					const FSkelMeshRenderSection& Section = LODRenderData->RenderSections[SectionIndex];
+					int32 BufferVertIndex = Section.GetVertexBufferIndex() + VertIndex;
+					for (int32 InfluenceIndex = 0; InfluenceIndex < MaxBoneInfluences; InfluenceIndex++)
+					{
+						const int32 BoneIndex = Section.BoneMap[SkinWeightVertexBuffer->GetBoneIndex(BufferVertIndex, InfluenceIndex)];
+						const float	Weight = (float)SkinWeightVertexBuffer->GetBoneWeight(BufferVertIndex, InfluenceIndex) * UE::AnimationCore::InvMaxRawBoneWeightFloat;
+						if (Weight > float(0) && 0 <= BoneIndex && BoneIndex < ComponentPose.Num())
+						{
+							FString BoneName = SkeletalMesh->GetRefSkeleton().GetBoneName(BoneIndex).ToString();
+							BoundVerts[BoneIndex].Add(index+j);
+							BoundWeights[BoneIndex].Add(Weight);
+						}
+					}
+				}
+				for (int32 BoneIndex = 0; BoneIndex < ComponentPose.Num(); ++BoneIndex)
+				{
+					if (BoundVerts[BoneIndex].Num())
+					{
+						FString BoneName = SkeletalMesh->GetRefSkeleton().GetBoneName(BoneIndex).ToString();
+						//get local coords of bound verts
+						typedef GeometryCollection::Facades::FKinematicBindingFacade FKinematics;
+						FKinematics Kinematics(InCollection); Kinematics.DefineSchema();
+						if (Kinematics.IsValid())
+						{
+							FKinematics::FBindingKey Binding = Kinematics.SetBoneBindings(BoneIndex, BoundVerts[BoneIndex], BoundWeights[BoneIndex]);
+							TManagedArray<TArray<FVector3f>>& LocalPos = InCollection.AddAttribute<TArray<FVector3f>>("LocalPosition", Binding.GroupName);
+							Kinematics.AddKinematicBinding(Binding);
+
+							auto FloatVert = [](FVector3d V) { return FVector3f(V.X, V.Y, V.Z); };
+							auto DoubleVert = [](FVector3f V) { return FVector3d(V.X, V.Y, V.Z); };
+							LocalPos[Binding.Index].SetNum(BoundVerts[BoneIndex].Num());
+							for (int32 i = 0; i < BoundVerts[BoneIndex].Num(); i++)
+							{
+								FVector3f Temp = Vertices[BoundVerts[BoneIndex][i]];
+								LocalPos[Binding.Index][i] = FloatVert(ComponentPose[BoneIndex].InverseTransformPosition(DoubleVert(Temp)));
+							}
+						}
+					}
+				}
+				GeometryCollection::Facades::FVertexBoneWeightsFacade(InCollection).AddBoneWeightsFromKinematicBindings();
+			}
+		}
+		SetValue<DataType>(Context, InCollection, &Collection);
+		SetValue<TArray<int32>>(Context, Indices, &IndicesOut);
+	}
 }
\ No newline at end of file
diff --git a/Engine/Plugins/Experimental/ChaosFlesh/Source/ChaosFleshNodes/Private/Dataflow/ChaosFleshPositionTargetInitializationNodes.cpp b/Engine/Plugins/Experimental/ChaosFlesh/Source/ChaosFleshNodes/Private/Dataflow/ChaosFleshPositionTargetInitializationNodes.cpp
--- a/Engine/Plugins/Experimental/ChaosFlesh/Source/ChaosFleshNodes/Private/Dataflow/ChaosFleshPositionTargetInitializationNodes.cpp
+++ b/Engine/Plugins/Experimental/ChaosFlesh/Source/ChaosFleshNodes/Private/Dataflow/ChaosFleshPositionTargetInitializationNodes.cpp
@@ -34,14 +34,12 @@ namespace Dataflow
 
 void FAddKinematicParticlesDataflowNode::Evaluate(Dataflow::FContext& Context, const FDataflowOutput* Out) const
 {
-	if (Out->IsA<DataType>(&Collection))
+	if (Out->IsA<DataType>(&Collection) || Out->IsA<TArray<int32>>(&TargetIndicesOut))
 	{
 		DataType InCollection = GetValue<DataType>(Context, &Collection);
-
+		TArray<int32> TargetIndices;
 		if (TManagedArray<FVector3f>* Vertices = InCollection.FindAttribute<FVector3f>("Vertex", FGeometryCollection::VerticesGroup))
 		{
-			TArray<int32> TargetIndices;
-			TargetIndices.SetNum(0);
 			if (FindInput(&VertexIndicesIn) && FindInput(&VertexIndicesIn)->GetConnection())
 			{
 				TArray<int32> BoundVerts;
@@ -139,10 +137,9 @@ void FAddKinematicParticlesDataflowNode::Evaluate(Dataflow::FContext& Context, c
 					GeometryCollection::Facades::FVertexBoneWeightsFacade(InCollection).AddBoneWeightsFromKinematicBindings();
 				}
 			}
-
-			SetValue<TArray<int32>>(Context, TargetIndices, &TargetIndicesOut);
 		}
 		SetValue<DataType>(Context, InCollection, &Collection);
+		SetValue<TArray<int32>>(Context, TargetIndices, &TargetIndicesOut);
 	}
 }
 
@@ -323,7 +320,7 @@ void FSetVertexTetrahedraPositionTargetBindingDataflowNode::Evaluate(Dataflow::F
 									for (int32 j = 0; j < TetIntersections.Num(); j++)
 									{
 										const int32 TetIdx = TetIntersections[j];
-										if (!Tets[TetIdx].Outside(ParticlePos, 1.0e-2))
+										if (!Tets[TetIdx].Outside(ParticlePos, 0))
 										{
 											Chaos::TVector<Chaos::FReal, 4> WeightsD = Tets[TetIdx].GetBarycentricCoordinates(ParticlePos);
 											GeometryCollection::Facades::FPositionTargetsData DataPackage;
diff --git a/Engine/Plugins/Experimental/ChaosFlesh/Source/ChaosFleshNodes/Public/Dataflow/ChaosFleshKinematicInitializationNodes.h b/Engine/Plugins/Experimental/ChaosFlesh/Source/ChaosFleshNodes/Public/Dataflow/ChaosFleshKinematicInitializationNodes.h
--- a/Engine/Plugins/Experimental/ChaosFlesh/Source/ChaosFleshNodes/Public/Dataflow/ChaosFleshKinematicInitializationNodes.h
+++ b/Engine/Plugins/Experimental/ChaosFlesh/Source/ChaosFleshNodes/Public/Dataflow/ChaosFleshKinematicInitializationNodes.h
@@ -88,14 +88,48 @@ public:
 		RegisterInputConnection(&Collection);
 		RegisterOutputConnection(&Collection, &Collection);
 		RegisterInputConnection(&SkeletalMeshIn);
-		RegisterInputConnection(&VertexIndicesIn);	
+		RegisterInputConnection(&VertexIndicesIn);
 		RegisterInputConnection(&BoneIndexIn);
 	}
 
 	virtual void Evaluate(Dataflow::FContext& Context, const FDataflowOutput* Out) const override;
 };
 
+USTRUCT(meta = (DataflowFlesh))
+struct FKinematicOriginInsertionInitializationDataflowNode : public FDataflowNode
+{
+	GENERATED_USTRUCT_BODY()
+		DATAFLOW_NODE_DEFINE_INTERNAL(FKinematicOriginInsertionInitializationDataflowNode, "KinematicOriginInsertionInitialization", "Flesh", "")
+		DATAFLOW_NODE_RENDER_TYPE(FGeometryCollection::StaticType(), "Collection")
+
+public:
+	typedef FManagedArrayCollection DataType;
+
+	UPROPERTY(meta = (DataflowInput, DataflowOutput, DisplayName = "Collection", DataflowPassthrough = "Collection"))
+		FManagedArrayCollection Collection;
+
+	UPROPERTY(meta = (DataflowInput, DisplayName = "OriginSelectionSet"))
+		TArray<int32> OriginVertexIndicesIn;
 
+	UPROPERTY(meta = (DataflowInput, DisplayName = "InsertionSelectionSet"))
+		TArray<int32> InsertionVertexIndicesIn;
+
+	UPROPERTY(meta = (DataflowInput, DisplayName = "BoneSkeletalMesh"))
+		TObjectPtr<USkeletalMesh> BoneSkeletalMeshIn;
+
+
+	FKinematicOriginInsertionInitializationDataflowNode(const Dataflow::FNodeParameters& InParam, FGuid InGuid = FGuid::NewGuid())
+		: FDataflowNode(InParam, InGuid)
+	{
+		RegisterInputConnection(&Collection);
+		RegisterOutputConnection(&Collection, &Collection);
+		RegisterInputConnection(&OriginVertexIndicesIn);
+		RegisterInputConnection(&InsertionVertexIndicesIn);
+		RegisterInputConnection(&BoneSkeletalMeshIn);
+	}
+
+	virtual void Evaluate(Dataflow::FContext& Context, const FDataflowOutput* Out) const override;
+};
 
 USTRUCT(meta = (DataflowFlesh))
 struct FSetVerticesKinematicDataflowNode : public FDataflowNode
@@ -153,6 +187,37 @@ public:
 	virtual void Evaluate(Dataflow::FContext& Context, const FDataflowOutput* Out) const override;
 };
 
+USTRUCT(meta = (DataflowFlesh))
+struct FKinematicSkeletalMeshInitializationDataflowNode : public FDataflowNode
+{
+	GENERATED_USTRUCT_BODY()
+		DATAFLOW_NODE_DEFINE_INTERNAL(FKinematicSkeletalMeshInitializationDataflowNode, "KinematicSkeletalMeshInitialization", "Flesh", "")
+		//DATAFLOW_NODE_RENDER_TYPE(FGeometryCollection::StaticType(), "Collection")
+
+public:
+	typedef FManagedArrayCollection DataType;
+
+	UPROPERTY(meta = (DataflowInput, DataflowOutput, DisplayName = "Collection"))
+		FManagedArrayCollection Collection;
+
+	UPROPERTY(meta = (DataflowInput, DisplayName = "SkeletalMesh"))
+		TObjectPtr<USkeletalMesh> SkeletalMeshIn;
+	
+	UPROPERTY(meta = (DataflowOutput, DisplayName = "SelectionSet"))
+		TArray<int32> IndicesOut;
+
+	FKinematicSkeletalMeshInitializationDataflowNode(const Dataflow::FNodeParameters& InParam, FGuid InGuid = FGuid::NewGuid())
+		: FDataflowNode(InParam, InGuid)
+	{
+		RegisterInputConnection(&Collection);
+		RegisterInputConnection(&SkeletalMeshIn);
+		RegisterOutputConnection(&Collection);
+		RegisterOutputConnection(&IndicesOut);
+	}
+
+	virtual void Evaluate(Dataflow::FContext& Context, const FDataflowOutput* Out) const override;
+};
+
 namespace Dataflow
 {
 	void RegisterChaosFleshKinematicInitializationNodes();
diff --git a/Engine/Source/Runtime/Experimental/Chaos/Private/Chaos/Deformable/ChaosDeformableSolver.cpp b/Engine/Source/Runtime/Experimental/Chaos/Private/Chaos/Deformable/ChaosDeformableSolver.cpp
--- a/Engine/Source/Runtime/Experimental/Chaos/Private/Chaos/Deformable/ChaosDeformableSolver.cpp
+++ b/Engine/Source/Runtime/Experimental/Chaos/Private/Chaos/Deformable/ChaosDeformableSolver.cpp
@@ -252,11 +252,14 @@ namespace Chaos::Softs
 		{
 			if (Proxy->CanSimulate())
 			{
-				InitializeDeformableParticles(*Proxy);
-				InitializeKinematicParticles(*Proxy);
-				InitializeWeakConstraint(*Proxy);
-				InitializeTetrahedralConstraint(*Proxy);
-				InitializeGidBasedConstraints(*Proxy);
+				if (Proxy->GetRestCollection().NumElements(FGeometryCollection::VerticesGroup))
+				{
+					InitializeDeformableParticles(*Proxy);
+					InitializeKinematicParticles(*Proxy);
+					InitializeWeakConstraint(*Proxy);
+					InitializeTetrahedralConstraint(*Proxy);
+					InitializeGidBasedConstraints(*Proxy);
+				}
 			}
 		}
 		if (FCollisionManagerProxy* CollisionManagerProxy = InProxy.As< FCollisionManagerProxy>())
@@ -803,74 +806,38 @@ namespace Chaos::Softs
 										int32 TransformNum = FleshInputBuffer->Transforms.Num();
 										if (RestNum > 0 && TransformNum > 0)
 										{
-											
-											for (int32 i = BoneIndices.Num() - 1; 0 <= i; i--)
+
+											for (int32 i = 0; i < BoneIndices.Num(); i++)
 											{
 												if (BoneIndices[i] > -1 && BoneIndices[i] < RestNum && BoneIndices[i] < TransformNum)
 												{
-													// We're lerp'ing between the previous skinned position and the current skinned 
-													// position, based on the substep.  If our SimSpace is a bone, we need to transform
-													// those points by the inverse bone transform at each of those points in time.
-													const FVector3f& Pt = Vertex[Index - Range[0]];
 
 													// @todo(flesh) : Add the pre-cached component space rest transforms to the rest collection. 
 													// see  UFleshComponent::NewDeformableData for how its pulled from the SkeletalMesh
-													const FTransform& BoneRestTransform = FleshInputBuffer->RestTransforms[BoneIndices[i]];
-													FVec3 LocalPoint = BoneRestTransform.InverseTransformPosition(ChaosVert(Pt));
-													const FTransform& BoneCurrentTransform = FleshInputBuffer->Transforms[BoneIndices[i]];
-													FVec3 ComponentPointAtT = BoneCurrentTransform.TransformPosition(LocalPoint);
+													FVec3 LocalPoint = FleshInputBuffer->RestTransforms[BoneIndices[i]].InverseTransformPosition(ChaosVert(Vertex[Index - Range[0]]));
+													FVec3 ComponentPointAtT = FleshInputBuffer->Transforms[BoneIndices[i]].TransformPosition(LocalPoint);
 
 													if (PreviousFleshBuffer)
 													{
-														const FTransform& BonePreviousTransform = PreviousFleshBuffer->Transforms[BoneIndices[i]];
-														FVec3 PreviousComponentPointAtT = BonePreviousTransform.TransformPosition(LocalPoint);
-
-														if (Proxy->IsBoneSpace())
-														{
-															// Current position gets current bone xf
-															ComponentPointAtT = GlobalTransform.TransformPosition(ComponentPointAtT);
-															// Previous position gets previous bone xf
-															const FTransform& PrevGlobalTransform = Proxy->GetPreviousPointsTransform();
-															PreviousComponentPointAtT = PrevGlobalTransform.TransformPosition(PreviousComponentPointAtT);
-														}
-														ComponentPointAtT = 
-															ComponentPointAtT * CurrentRatio + PreviousComponentPointAtT * ((FSolverReal)1. - CurrentRatio);
-													}
-													else if (Proxy->IsBoneSpace())
-													{
-														ComponentPointAtT = GlobalTransform.TransformPosition(ComponentPointAtT);
+														FTransform BonePreviousTransform = PreviousFleshBuffer->Transforms[BoneIndices[i]];
+														ComponentPointAtT = ComponentPointAtT * CurrentRatio + BonePreviousTransform.TransformPosition(LocalPoint) * ((FSolverReal)1. - CurrentRatio);
 													}
 
-													if (!Proxy->IsBoneSpace())
-													{
-														MParticles.X(Index) = GlobalTransform.TransformPosition(ComponentPointAtT);
-													}
-													else
-													{
-														MParticles.X(Index) = ComponentPointAtT;
-													}
+													MParticles.X(Index) = MParticles.X(Index) + GlobalTransform.TransformPosition(ComponentPointAtT) * BoneWeights[i];
 
 													bParticleTouched = true;
-
-
-													// @todo(flesh): Add non rigid skinning weights. 
-													// Currently this just grabs the last joint, ideally
-													// we would be using the weighted average for the vertex
-													// positions.
-
-#if WITH_EDITOR
-													//debug draw
-													//p.Chaos.DebugDraw.Enabled 1
-													//p.Chaos.DebugDraw.Deformable.KinematicParticle 1
-													if (GDeformableDebugParams.IsDebugDrawingEnabled() && GDeformableDebugParams.bDoDrawKinematicParticles)
-													{
-														auto DoubleVert = [](FVector3f V) { return FVector3d(V.X, V.Y, V.Z); };
-														Chaos::FDebugDrawQueue::GetInstance().DrawDebugPoint(DoubleVert(MParticles.X(Index)), FColor::Red, false, -1.0f, 0, 5);
-													}
-#endif
-													break;
 												}
 											}
+#if WITH_EDITOR
+											//debug draw
+											//p.Chaos.DebugDraw.Enabled 1
+											//p.Chaos.DebugDraw.Deformable.KinematicParticle 1
+											if (GDeformableDebugParams.IsDebugDrawingEnabled() && GDeformableDebugParams.bDoDrawKinematicParticles)
+											{
+												auto DoubleVert = [](FVector3f V) { return FVector3d(V.X, V.Y, V.Z); };
+												Chaos::FDebugDrawQueue::GetInstance().DrawDebugPoint(DoubleVert(MParticles.X(Index)), FColor::Red, false, -1.0f, 0, 5);
+											}
+#endif
 										}
 										MParticles.PAndInvM(Index).P = MParticles.X(Index);
 									}
