ChaosFlesh : Add blend to animation on output samples

diff --git a/Engine/Plugins/Experimental/ChaosFlesh/Source/ChaosFleshEngine/Private/ChaosFlesh/FleshComponent.cpp b/Engine/Plugins/Experimental/ChaosFlesh/Source/ChaosFleshEngine/Private/ChaosFlesh/FleshComponent.cpp
--- a/Engine/Plugins/Experimental/ChaosFlesh/Source/ChaosFleshEngine/Private/ChaosFlesh/FleshComponent.cpp
+++ b/Engine/Plugins/Experimental/ChaosFlesh/Source/ChaosFleshEngine/Private/ChaosFlesh/FleshComponent.cpp
@@ -27,6 +27,7 @@
 FChaosEngineDeformableCVarParams CVarParams;
 FAutoConsoleVariableRef CVarDeforambleDoDrawSimulationMesh(TEXT("p.Chaos.DebugDraw.Deformable.SimulationMesh"), CVarParams.bDoDrawSimulationMesh, TEXT("Debug draw the deformable simulation resutls on the game thread. [def: true]"));
 FAutoConsoleVariableRef CVarDeforambleDoDrawSkeletalMeshBindingPositions(TEXT("p.Chaos.DebugDraw.Deformable.SkeletalMeshBindingPositions"), CVarParams.bDoDrawSkeletalMeshBindingPositions, TEXT("Debug draw the deformable simulation's SkeletalMeshBindingPositions on the game thread. [def: false]"));
+FAutoConsoleVariableRef CVarDeforambleDoDrawSkeletalMeshBindingPositionsAnimationBlendWeight(TEXT("p.Chaos.DebugDraw.Deformable.SkeletalMeshBindingPositions.AnimationBlendWeight"), CVarParams.DrawSkeletalMeshBindingPositionsAnimationBlendWeight, TEXT("Set the animation blend weight of the skeletal mesh debug draw.[def: 0.]"));
 
 #define PERF_SCOPE(X) SCOPE_CYCLE_COUNTER(X); TRACE_CPUPROFILER_EVENT_SCOPE(X);
 DECLARE_CYCLE_STAT(TEXT("Chaos.Deformable.UFleshComponent.TickComponent"), STAT_ChaosDeformable_UFleshComponent_TickComponent, STATGROUP_Chaos);
@@ -542,6 +543,7 @@ void UFleshComponent::DebugDrawSkeletalMeshBindingPositions() const
 {
 #if WITH_EDITOR
 	auto UEVertd = [](FVector3f V) { return FVector3d(V.X, V.Y, V.Z); };
+	float AnimationBlendWeight = CVarParams.DrawSkeletalMeshBindingPositionsAnimationBlendWeight;
 
 	if (const UFleshAsset* RestAsset = GetRestCollection())
 	{
@@ -552,7 +554,7 @@ void UFleshComponent::DebugDrawSkeletalMeshBindingPositions() const
 		{
 			TArray<bool> Influenced;
 			TArray<FVector> PosArray = GetSkeletalMeshEmbeddedPositionsInternal(
-				ChaosDeformableBindingOption::ComponentPos, FTransform::Identity, "", &Influenced);
+				ChaosDeformableBindingOption::ComponentPos, FTransform::Identity, "", AnimationBlendWeight, &Influenced);
 
 			for (int i = 0; i < PosArray.Num(); i++)
 			{
@@ -576,15 +578,17 @@ TArray<FVector> UFleshComponent::GetSkeletalMeshBindingPositions(const USkeletal
 TArray<FVector> UFleshComponent::GetSkeletalMeshEmbeddedPositions(
 	const ChaosDeformableBindingOption Format,
 	const FTransform TargetDeformationSkeletonOffset,
-	const FName TargetBone) const
+	const FName TargetBone,
+	const float AnimationBlendWeight) const
 {
-	return GetSkeletalMeshEmbeddedPositionsInternal(Format, TargetDeformationSkeletonOffset, TargetBone, nullptr);
+	return GetSkeletalMeshEmbeddedPositionsInternal(Format, TargetDeformationSkeletonOffset, TargetBone, AnimationBlendWeight, nullptr);
 }
 
 TArray<FVector> UFleshComponent::GetSkeletalMeshEmbeddedPositionsInternal(
 	const ChaosDeformableBindingOption Format,
 	const FTransform TargetDeformationSkeletonOffset, 
 	const FName TargetBone,
+	const float AnimationBlendWeight,
 	TArray<bool>* OutInfluence) const
 {
 	if (!GetRestCollection())
@@ -642,6 +646,7 @@ TArray<FVector> UFleshComponent::GetSkeletalMeshEmbeddedPositionsInternal(
 		GetEmbeddedPositionsInternal(
 			TransformPositions, 
 			FName(GetRestCollection()->TargetDeformationSkeleton->GetName()), // for identifying the binding group
+			AnimationBlendWeight,
 			OutInfluence);
 
 	// World space
@@ -788,6 +793,7 @@ TArray<FVector> UFleshComponent::GetSkeletalMeshEmbeddedPositionsInternal(
 TArray<FVector> UFleshComponent::GetEmbeddedPositionsInternal(
 	const TArray<FVector>& InPositions, 
 	const FName SkeletalMeshName,
+	const float AnimationBlendWeight,
 	TArray<bool>* OutInfluence) const
 {
 	auto UEVert3d = [](FVector3f V) { return FVector3d(V.X, V.Y, V.Z); };
@@ -807,23 +813,40 @@ TArray<FVector> UFleshComponent::GetEmbeddedPositionsInternal(
 
 			if (ensure(Verts != nullptr) && TetrahedronStart)
 			{
-				if (OutInfluence != nullptr)
-					OutInfluence->Init(false, OutPositions.Num());
-				for (int32 TetMeshIdx = 0; TetMeshIdx < TetrahedronStart->Num(); TetMeshIdx++)
+				auto CalculateBindings = [this, TetBindings, TetrahedronStart, SkeletalMeshName, OutInfluence](const TManagedArray<FVector3f>* Verts, TArray<FVector>& OutPositions)
 				{
-					FString MeshBindingsName =
-						GeometryCollection::Facades::FTetrahedralSkeletalBindings::GenerateMeshGroupName(TetMeshIdx, SkeletalMeshName);
-					if (!TetBindings.CalculateBindings(MeshBindingsName, Verts->GetConstArray(), OutPositions, OutInfluence))
+					if (OutInfluence != nullptr)
+						OutInfluence->Init(false, OutPositions.Num());
+					for (int32 TetMeshIdx = 0; TetMeshIdx < TetrahedronStart->Num(); TetMeshIdx++)
 					{
-						UE_LOG(LogFleshComponentInternal, Warning,
-							TEXT("'%s' - GetEmbeddedPositionsInternal - Failed to find mesh bindings for skeletal mesh '%s'"),
-							*GetName(),
-							*SkeletalMeshName.ToString());
+						FString MeshBindingsName =
+							GeometryCollection::Facades::FTetrahedralSkeletalBindings::GenerateMeshGroupName(TetMeshIdx, SkeletalMeshName);
+						if (!TetBindings.CalculateBindings(MeshBindingsName, Verts->GetConstArray(), OutPositions, OutInfluence))
+						{
+							UE_LOG(LogFleshComponentInternal, Warning,
+								   TEXT("'%s' - GetEmbeddedPositionsInternal - Failed to find mesh bindings for skeletal mesh '%s'"),
+								   *GetName(),
+								   *SkeletalMeshName.ToString());
+						}
 					}
-				}
+				};
+				CalculateBindings(Verts, OutPositions);
 
-			}
+				// blend between the aniamtion position and the 
+				// simulated position. 
+				if (!FMath::IsNearlyZero(AnimationBlendWeight) && GetDynamicCollection())
+				{
+					TArray<FVector> RestPositions = InPositions;
+					CalculateBindings(RestVerts, RestPositions);
 
+					float ClampedWeight = FMath::Clamp(AnimationBlendWeight, 0.f, 1.f);
+					for (int i = 0; i < OutPositions.Num(); i++)
+					{
+						FVector V = RestPositions[i] - OutPositions[i];
+						OutPositions[i] += V * ClampedWeight;
+					}
+				}
+			}
 		}
 	}
 	return OutPositions;
diff --git a/Engine/Plugins/Experimental/ChaosFlesh/Source/ChaosFleshEngine/Public/ChaosFlesh/ChaosDeformableTypes.h b/Engine/Plugins/Experimental/ChaosFlesh/Source/ChaosFleshEngine/Public/ChaosFlesh/ChaosDeformableTypes.h
--- a/Engine/Plugins/Experimental/ChaosFlesh/Source/ChaosFleshEngine/Public/ChaosFlesh/ChaosDeformableTypes.h
+++ b/Engine/Plugins/Experimental/ChaosFlesh/Source/ChaosFleshEngine/Public/ChaosFlesh/ChaosDeformableTypes.h
@@ -26,4 +26,5 @@ struct FChaosEngineDeformableCVarParams
 	bool bDoDrawSimulationMesh = false;
 #endif
 	bool bDoDrawSkeletalMeshBindingPositions = false;
+	float DrawSkeletalMeshBindingPositionsAnimationBlendWeight = 0.f;
 };
\ No newline at end of file
diff --git a/Engine/Plugins/Experimental/ChaosFlesh/Source/ChaosFleshEngine/Public/ChaosFlesh/FleshComponent.h b/Engine/Plugins/Experimental/ChaosFlesh/Source/ChaosFleshEngine/Public/ChaosFlesh/FleshComponent.h
--- a/Engine/Plugins/Experimental/ChaosFlesh/Source/ChaosFleshEngine/Public/ChaosFlesh/FleshComponent.h
+++ b/Engine/Plugins/Experimental/ChaosFlesh/Source/ChaosFleshEngine/Public/ChaosFlesh/FleshComponent.h
@@ -85,7 +85,7 @@ public:
 	* the \c TargetDeformationSkeleton to be co-located with the flesh mesh.
 	*/
 	UFUNCTION(BlueprintCallable, Category = "Physics")
-	TArray<FVector> GetSkeletalMeshEmbeddedPositions(const ChaosDeformableBindingOption Format, const FTransform TargetDeformationSkeletonOffset, const FName TargetBone = "") const;
+	TArray<FVector> GetSkeletalMeshEmbeddedPositions(const ChaosDeformableBindingOption Format, const FTransform TargetDeformationSkeletonOffset, const FName TargetBone = "", const float AnimationBlendWeight = 0.f) const;
 
 	/** 
 	* Offset transform that moves the \c TargetDeformationSkeleton to be co-located with the flesh mesh. 
@@ -155,8 +155,8 @@ private:
 
 	FTransform PrevTransform = FTransform::Identity;
 
-	TArray<FVector> GetSkeletalMeshEmbeddedPositionsInternal(const ChaosDeformableBindingOption Format, const FTransform TargetDeformationSkeletonOffset, const FName TargetBone = "", TArray<bool>* OutInfluence = nullptr) const;
-	TArray<FVector> GetEmbeddedPositionsInternal(const TArray<FVector>& InPositions, const FName SkeletalMeshName, TArray<bool>* OutInfluence = nullptr) const;
+	TArray<FVector> GetSkeletalMeshEmbeddedPositionsInternal(const ChaosDeformableBindingOption Format, const FTransform TargetDeformationSkeletonOffset, const FName TargetBone = "", const float AnimationBlendWeight = 0.f, TArray<bool>* OutInfluence = nullptr) const;
+	TArray<FVector> GetEmbeddedPositionsInternal(const TArray<FVector>& InPositions, const FName SkeletalMeshName, const float AnimationBlendWeight = 0.f, TArray<bool>* OutInfluence = nullptr) const;
 	TArray<FVector> GetSkeletalMeshBindingPositionsInternal(const USkeletalMesh* InSkeletalMesh, TArray<bool>* OutInfluence = nullptr) const;
 	void DebugDrawSkeletalMeshBindingPositions() const;
 
