ChaosFlesh : Simulation Blend Weight for deltas

diff --git a/Engine/Plugins/Experimental/ChaosFlesh/Source/ChaosFleshEngine/Private/ChaosFlesh/ChaosDeformablePhysicsComponent.cpp b/Engine/Plugins/Experimental/ChaosFlesh/Source/ChaosFleshEngine/Private/ChaosFlesh/ChaosDeformablePhysicsComponent.cpp
--- a/Engine/Plugins/Experimental/ChaosFlesh/Source/ChaosFleshEngine/Private/ChaosFlesh/ChaosDeformablePhysicsComponent.cpp
+++ b/Engine/Plugins/Experimental/ChaosFlesh/Source/ChaosFleshEngine/Private/ChaosFlesh/ChaosDeformablePhysicsComponent.cpp
@@ -7,12 +7,9 @@
 #include "ChaosFlesh/FleshComponent.h"
 
 
-DEFINE_LOG_CATEGORY_STATIC(LogDeformablePhysicsComponentInternal, Log, All);
-
 void UDeformablePhysicsComponent::OnCreatePhysicsState()
 {
 	Super::OnCreatePhysicsState();
-	UE_LOG(LogDeformablePhysicsComponentInternal, Log, TEXT("UDeformablePhysicsComponent::OnCreatePhysicsState()"));
 	if (PrimarySolverComponent)
 	{
 		FDeformableSolver::FGameThreadAccess GameThreadSolver = PrimarySolverComponent->GameThreadAccess();
@@ -26,7 +23,6 @@ void UDeformablePhysicsComponent::OnCreatePhysicsState()
 void UDeformablePhysicsComponent::OnDestroyPhysicsState()
 {
 	Super::OnDestroyPhysicsState();
-	UE_LOG(LogDeformablePhysicsComponentInternal, Log, TEXT("UDeformablePhysicsComponent::OnDestroyPhysicsState()"));
 	if (PrimarySolverComponent)
 	{
 		FDeformableSolver::FGameThreadAccess GameThreadSolver = PrimarySolverComponent->GameThreadAccess();
diff --git a/Engine/Plugins/Experimental/ChaosFlesh/Source/ChaosFleshEngine/Private/ChaosFlesh/FleshComponent.cpp b/Engine/Plugins/Experimental/ChaosFlesh/Source/ChaosFleshEngine/Private/ChaosFlesh/FleshComponent.cpp
--- a/Engine/Plugins/Experimental/ChaosFlesh/Source/ChaosFleshEngine/Private/ChaosFlesh/FleshComponent.cpp
+++ b/Engine/Plugins/Experimental/ChaosFlesh/Source/ChaosFleshEngine/Private/ChaosFlesh/FleshComponent.cpp
@@ -27,7 +27,7 @@
 FChaosEngineDeformableCVarParams CVarParams;
 FAutoConsoleVariableRef CVarDeforambleDoDrawSimulationMesh(TEXT("p.Chaos.DebugDraw.Deformable.SimulationMesh"), CVarParams.bDoDrawSimulationMesh, TEXT("Debug draw the deformable simulation resutls on the game thread. [def: true]"));
 FAutoConsoleVariableRef CVarDeforambleDoDrawSkeletalMeshBindingPositions(TEXT("p.Chaos.DebugDraw.Deformable.SkeletalMeshBindingPositions"), CVarParams.bDoDrawSkeletalMeshBindingPositions, TEXT("Debug draw the deformable simulation's SkeletalMeshBindingPositions on the game thread. [def: false]"));
-FAutoConsoleVariableRef CVarDeforambleDoDrawSkeletalMeshBindingPositionsAnimationBlendWeight(TEXT("p.Chaos.DebugDraw.Deformable.SkeletalMeshBindingPositions.AnimationBlendWeight"), CVarParams.DrawSkeletalMeshBindingPositionsAnimationBlendWeight, TEXT("Set the animation blend weight of the skeletal mesh debug draw.[def: 0.]"));
+FAutoConsoleVariableRef CVarDeforambleDoDrawSkeletalMeshBindingPositionsSimulationBlendWeight(TEXT("p.Chaos.DebugDraw.Deformable.SkeletalMeshBindingPositions.SimulationBlendWeight"), CVarParams.DrawSkeletalMeshBindingPositionsSimulationBlendWeight, TEXT("Set the simulation blend weight of the skeletal mesh debug draw.[def: 1.]"));
 
 #define PERF_SCOPE(X) SCOPE_CYCLE_COUNTER(X); TRACE_CPUPROFILER_EVENT_SCOPE(X);
 DECLARE_CYCLE_STAT(TEXT("Chaos.Deformable.UFleshComponent.TickComponent"), STAT_ChaosDeformable_UFleshComponent_TickComponent, STATGROUP_Chaos);
@@ -543,7 +543,7 @@ void UFleshComponent::DebugDrawSkeletalMeshBindingPositions() const
 {
 #if WITH_EDITOR
 	auto UEVertd = [](FVector3f V) { return FVector3d(V.X, V.Y, V.Z); };
-	float AnimationBlendWeight = CVarParams.DrawSkeletalMeshBindingPositionsAnimationBlendWeight;
+	float SimulationBlendWeight = CVarParams.DrawSkeletalMeshBindingPositionsSimulationBlendWeight;
 
 	if (const UFleshAsset* RestAsset = GetRestCollection())
 	{
@@ -554,7 +554,7 @@ void UFleshComponent::DebugDrawSkeletalMeshBindingPositions() const
 		{
 			TArray<bool> Influenced;
 			TArray<FVector> PosArray = GetSkeletalMeshEmbeddedPositionsInternal(
-				ChaosDeformableBindingOption::ComponentPos, FTransform::Identity, "", AnimationBlendWeight, &Influenced);
+				ChaosDeformableBindingOption::ComponentPos, FTransform::Identity, "", SimulationBlendWeight, &Influenced);
 
 			for (int i = 0; i < PosArray.Num(); i++)
 			{
@@ -579,16 +579,16 @@ TArray<FVector> UFleshComponent::GetSkeletalMeshEmbeddedPositions(
 	const ChaosDeformableBindingOption Format,
 	const FTransform TargetDeformationSkeletonOffset,
 	const FName TargetBone,
-	const float AnimationBlendWeight) const
+	const float SimulationBlendWeight) const
 {
-	return GetSkeletalMeshEmbeddedPositionsInternal(Format, TargetDeformationSkeletonOffset, TargetBone, AnimationBlendWeight, nullptr);
+	return GetSkeletalMeshEmbeddedPositionsInternal(Format, TargetDeformationSkeletonOffset, TargetBone, SimulationBlendWeight, nullptr);
 }
 
 TArray<FVector> UFleshComponent::GetSkeletalMeshEmbeddedPositionsInternal(
 	const ChaosDeformableBindingOption Format,
 	const FTransform TargetDeformationSkeletonOffset, 
 	const FName TargetBone,
-	const float AnimationBlendWeight,
+	const float SimulationBlendWeight,
 	TArray<bool>* OutInfluence) const
 {
 	if (!GetRestCollection())
@@ -606,6 +606,8 @@ TArray<FVector> UFleshComponent::GetSkeletalMeshEmbeddedPositionsInternal(
 		return TArray<FVector>();
 	}
 
+	TArray<FVector> EmbeddedPosComp;
+
 	// Get sample points in the skel mesh's component space.  This code assumes that the
 	// skeletal mesh and the flesh asset are aligned in their respective local spaces.  If
 	// they're not aligned, then the TargetDeformationSkeletonOffset should be provided to
@@ -641,17 +643,19 @@ TArray<FVector> UFleshComponent::GetSkeletalMeshEmbeddedPositionsInternal(
 		TransformPositions.SetNum(GetRestCollection()->TargetDeformationSkeleton->GetRefSkeleton().GetNum());
 	}
 
-	// Calculate their current embedded positions
-	TArray<FVector> EmbeddedPosComp = 
-		GetEmbeddedPositionsInternal(
-			TransformPositions, 
-			FName(GetRestCollection()->TargetDeformationSkeleton->GetName()), // for identifying the binding group
-			AnimationBlendWeight,
-			OutInfluence);
+
 
 	// World space
 	if (Format == ChaosDeformableBindingOption::WorldPos)
 	{
+		// Calculate their current embedded positions
+		EmbeddedPosComp =
+			GetEmbeddedPositionsInternal(
+			TransformPositions,
+			FName(GetRestCollection()->TargetDeformationSkeleton->GetName()), // for identifying the binding group
+			SimulationBlendWeight,
+			OutInfluence);
+
 		// Put component space points into world space
 		const FTransform ComponentXf = GetComponentTransform();
 		for (int32 i = 0; i < EmbeddedPosComp.Num(); i++)
@@ -661,6 +665,14 @@ TArray<FVector> UFleshComponent::GetSkeletalMeshEmbeddedPositionsInternal(
 	}
 	else if (Format == ChaosDeformableBindingOption::WorldDelta)
 	{
+		// Calculate their current embedded positions
+		EmbeddedPosComp =
+			GetEmbeddedPositionsInternal(
+			TransformPositions,
+			FName(GetRestCollection()->TargetDeformationSkeleton->GetName()), // for identifying the binding group
+			1.f,
+			OutInfluence);
+
 		const FTransform ComponentXf = GetComponentTransform();
 		for (int32 i = 0; i < EmbeddedPosComp.Num(); i++)
 		{
@@ -672,19 +684,40 @@ TArray<FVector> UFleshComponent::GetSkeletalMeshEmbeddedPositionsInternal(
 	// Component space
 	else if (Format == ChaosDeformableBindingOption::ComponentPos)
 	{
-		// Do nothing.
+		EmbeddedPosComp =
+			GetEmbeddedPositionsInternal(
+			TransformPositions,
+			FName(GetRestCollection()->TargetDeformationSkeleton->GetName()), // for identifying the binding group
+			1.f,
+			OutInfluence);
 	}
 	else if (Format == ChaosDeformableBindingOption::ComponentDelta)
 	{
+		// Calculate their current embedded positions
+		EmbeddedPosComp =
+			GetEmbeddedPositionsInternal(
+			TransformPositions,
+			FName(GetRestCollection()->TargetDeformationSkeleton->GetName()), // for identifying the binding group
+			1.f,
+			OutInfluence);
+
 		for (int32 i = 0; i < EmbeddedPosComp.Num(); i++)
 		{
-			EmbeddedPosComp[i] = EmbeddedPosComp[i] - TransformPositions[i];
+			EmbeddedPosComp[i] = (EmbeddedPosComp[i] - TransformPositions[i]) * SimulationBlendWeight;
 		}
 	}
 
 	// Bone space
 	else
 	{
+		EmbeddedPosComp =
+			GetEmbeddedPositionsInternal(
+			TransformPositions,
+			FName(GetRestCollection()->TargetDeformationSkeleton->GetName()), // for identifying the binding group
+			SimulationBlendWeight,
+			OutInfluence);
+
+
 		// Find the component that owns TargetDeformationSkeleton, so we can pull the 
 		// current animated bone transforms out of it.
 		//
@@ -793,7 +826,7 @@ TArray<FVector> UFleshComponent::GetSkeletalMeshEmbeddedPositionsInternal(
 TArray<FVector> UFleshComponent::GetEmbeddedPositionsInternal(
 	const TArray<FVector>& InPositions, 
 	const FName SkeletalMeshName,
-	const float AnimationBlendWeight,
+	const float SimulationBlendWeight,
 	TArray<bool>* OutInfluence) const
 {
 	auto UEVert3d = [](FVector3f V) { return FVector3d(V.X, V.Y, V.Z); };
@@ -834,16 +867,16 @@ TArray<FVector> UFleshComponent::GetEmbeddedPositionsInternal(
 
 				// blend between the aniamtion position and the 
 				// simulated position. 
-				if (!FMath::IsNearlyZero(AnimationBlendWeight) && GetDynamicCollection())
+				if (!FMath::IsNearlyEqual(SimulationBlendWeight, 1.0) && GetDynamicCollection())
 				{
 					TArray<FVector> RestPositions = InPositions;
 					CalculateBindings(RestVerts, RestPositions);
 
-					float ClampedWeight = FMath::Clamp(AnimationBlendWeight, 0.f, 1.f);
+					float ClampedWeight = FMath::Clamp(SimulationBlendWeight, 0.f, 1.f);
 					for (int i = 0; i < OutPositions.Num(); i++)
 					{
 						FVector V = RestPositions[i] - OutPositions[i];
-						OutPositions[i] += V * ClampedWeight;
+						OutPositions[i] += V * (1.0-ClampedWeight);
 					}
 				}
 			}
diff --git a/Engine/Plugins/Experimental/ChaosFlesh/Source/ChaosFleshEngine/Public/ChaosFlesh/ChaosDeformableTypes.h b/Engine/Plugins/Experimental/ChaosFlesh/Source/ChaosFleshEngine/Public/ChaosFlesh/ChaosDeformableTypes.h
--- a/Engine/Plugins/Experimental/ChaosFlesh/Source/ChaosFleshEngine/Public/ChaosFlesh/ChaosDeformableTypes.h
+++ b/Engine/Plugins/Experimental/ChaosFlesh/Source/ChaosFleshEngine/Public/ChaosFlesh/ChaosDeformableTypes.h
@@ -26,5 +26,5 @@ struct FChaosEngineDeformableCVarParams
 	bool bDoDrawSimulationMesh = false;
 #endif
 	bool bDoDrawSkeletalMeshBindingPositions = false;
-	float DrawSkeletalMeshBindingPositionsAnimationBlendWeight = 0.f;
+	float DrawSkeletalMeshBindingPositionsSimulationBlendWeight = 1.f;
 };
\ No newline at end of file
diff --git a/Engine/Plugins/Experimental/ChaosFlesh/Source/ChaosFleshEngine/Public/ChaosFlesh/FleshComponent.h b/Engine/Plugins/Experimental/ChaosFlesh/Source/ChaosFleshEngine/Public/ChaosFlesh/FleshComponent.h
--- a/Engine/Plugins/Experimental/ChaosFlesh/Source/ChaosFleshEngine/Public/ChaosFlesh/FleshComponent.h
+++ b/Engine/Plugins/Experimental/ChaosFlesh/Source/ChaosFleshEngine/Public/ChaosFlesh/FleshComponent.h
@@ -85,7 +85,7 @@ public:
 	* the \c TargetDeformationSkeleton to be co-located with the flesh mesh.
 	*/
 	UFUNCTION(BlueprintCallable, Category = "Physics")
-	TArray<FVector> GetSkeletalMeshEmbeddedPositions(const ChaosDeformableBindingOption Format, const FTransform TargetDeformationSkeletonOffset, const FName TargetBone = "", const float AnimationBlendWeight = 0.f) const;
+	TArray<FVector> GetSkeletalMeshEmbeddedPositions(const ChaosDeformableBindingOption Format, const FTransform TargetDeformationSkeletonOffset, const FName TargetBone = "", const float SimulationBlendWeight = 1.f) const;
 
 	/** 
 	* Offset transform that moves the \c TargetDeformationSkeleton to be co-located with the flesh mesh. 
@@ -155,8 +155,8 @@ private:
 
 	FTransform PrevTransform = FTransform::Identity;
 
-	TArray<FVector> GetSkeletalMeshEmbeddedPositionsInternal(const ChaosDeformableBindingOption Format, const FTransform TargetDeformationSkeletonOffset, const FName TargetBone = "", const float AnimationBlendWeight = 0.f, TArray<bool>* OutInfluence = nullptr) const;
-	TArray<FVector> GetEmbeddedPositionsInternal(const TArray<FVector>& InPositions, const FName SkeletalMeshName, const float AnimationBlendWeight = 0.f, TArray<bool>* OutInfluence = nullptr) const;
+	TArray<FVector> GetSkeletalMeshEmbeddedPositionsInternal(const ChaosDeformableBindingOption Format, const FTransform TargetDeformationSkeletonOffset, const FName TargetBone = "", const float SimulationBlendWeight = 1.f, TArray<bool>* OutInfluence = nullptr) const;
+	TArray<FVector> GetEmbeddedPositionsInternal(const TArray<FVector>& InPositions, const FName SkeletalMeshName, const float SimulationBlendWeight = 1.f, TArray<bool>* OutInfluence = nullptr) const;
 	TArray<FVector> GetSkeletalMeshBindingPositionsInternal(const USkeletalMesh* InSkeletalMesh, TArray<bool>* OutInfluence = nullptr) const;
 	void DebugDrawSkeletalMeshBindingPositions() const;
 
