ChaosFlesh : Add bulge parameter for radial meshing node.

diff --git a/Engine/Plugins/Experimental/ChaosFlesh/Source/ChaosFlesh/Private/Meshing/ChaosFleshRadialMeshing.cpp b/Engine/Plugins/Experimental/ChaosFlesh/Source/ChaosFlesh/Private/Meshing/ChaosFleshRadialMeshing.cpp
--- a/Engine/Plugins/Experimental/ChaosFlesh/Source/ChaosFlesh/Private/Meshing/ChaosFleshRadialMeshing.cpp
+++ b/Engine/Plugins/Experimental/ChaosFlesh/Source/ChaosFlesh/Private/Meshing/ChaosFleshRadialMeshing.cpp
@@ -3,21 +3,23 @@
 #include "Meshing/ChaosFleshRadialMeshing.h"
 
 using FReal = FVector::FReal;
-void RadialTetMesh(const FReal InnerRadius, const FReal OuterRadius, const FReal Height, const int32 RadialSample, const int32 AngularSample, const int32 VerticalSample, TArray<FIntVector4>& TetElements, TArray<FVector>& TetVertices) 
+void RadialTetMesh(const FReal InnerRadius, const FReal OuterRadius, const FReal Height, const int32 RadialSample, const int32 AngularSample, const int32 VerticalSample, const FReal BulgeDistance, TArray<FIntVector4>& TetElements, TArray<FVector>& TetVertices)
 {
 	TArray<FVector> HexVertices;
 	TArray<int32> HexElements;
-	RadialHexMesh(InnerRadius, OuterRadius, Height, RadialSample, AngularSample, VerticalSample, HexElements, HexVertices);
+	RadialHexMesh(InnerRadius, OuterRadius, Height, RadialSample, AngularSample, VerticalSample, BulgeDistance, HexElements, HexVertices);
 	RegularHexMesh2TetMesh(HexVertices, HexElements, TetVertices, TetElements);
 }
 
-void RadialHexMesh(const FReal InnerRadius, const FReal OuterRadius, const FReal Height, const int32 RadialSample, const int32 AngularSample, const int32 VerticalSample, TArray<int32>& HexElements, TArray<FVector>& HexVertices) 
+void RadialHexMesh(const FReal InnerRadius, const FReal OuterRadius, const FReal Height, const int32 RadialSample, const int32 AngularSample, const int32 VerticalSample, const FReal BulgeDistance, TArray<int32>& HexElements, TArray<FVector>& HexVertices) 
 {
 	HexElements.SetNum(8 * (VerticalSample - 1) * (RadialSample - 1) * AngularSample);
 	HexVertices.SetNum(VerticalSample * RadialSample * AngularSample);
 	FReal dr = (OuterRadius - InnerRadius) / FReal(RadialSample - 1);
 	FReal dtheta = FReal(2) * PI / FReal(AngularSample);
 	FReal dz = Height / FReal(VerticalSample - 1);
+	FReal l = (OuterRadius - InnerRadius) / (FReal)2.;
+	FReal Mag = (l * l + BulgeDistance * BulgeDistance) / ((FReal)2. * BulgeDistance);
 	for (int32 i = 0; i < VerticalSample; i++)
 	{
 		//start at top, work way down
@@ -29,6 +31,21 @@ void RadialHexMesh(const FReal InnerRadius, const FReal OuterRadius, const FReal
 			for (int32 k = 0; k < RadialSample; k++)
 			{
 				FReal RVal = InnerRadius + k * dr;
+
+				if (BulgeDistance > DOUBLE_SMALL_NUMBER)
+				{
+					if (i == 0)
+					{
+						FReal ZElevation = Height * FReal(0.5) - Mag * sin(PI * (Mag - l) / (2 * Mag));
+						ZVal = Mag * sin(PI * (Mag - l + (FReal)k * dr) / (2 * Mag)) + ZElevation;
+					}
+					if (i == VerticalSample - 1)
+					{
+						FReal ZElevation = -Height * FReal(0.5) - Mag * sin(PI + PI * (Mag - l) / (FReal(2) * Mag));
+						ZVal = Mag * sin(PI + PI * (Mag - l + (FReal)k * dr) / (FReal(2) * Mag)) + ZElevation;
+					}
+				}
+
 				HexVertices[AngularSample * RadialSample * i + RadialSample * j + k] = { RVal * XVal,RVal * YVal,ZVal };
 
 			}
diff --git a/Engine/Plugins/Experimental/ChaosFlesh/Source/ChaosFlesh/Public/Meshing/ChaosFleshRadialMeshing.h b/Engine/Plugins/Experimental/ChaosFlesh/Source/ChaosFlesh/Public/Meshing/ChaosFleshRadialMeshing.h
--- a/Engine/Plugins/Experimental/ChaosFlesh/Source/ChaosFlesh/Public/Meshing/ChaosFleshRadialMeshing.h
+++ b/Engine/Plugins/Experimental/ChaosFlesh/Source/ChaosFlesh/Public/Meshing/ChaosFleshRadialMeshing.h
@@ -2,8 +2,8 @@
 #pragma once
 #include "CoreMinimal.h"
 
-void CHAOSFLESH_API RadialTetMesh(const FVector::FReal InnerRadius, const FVector::FReal OuterRadius, const FVector::FReal Height, const int32 RadialSample, const int32 AngularSample, const int32 VerticalSample, TArray<FIntVector4>& TetElements, TArray<FVector>& TetVertices);
-void CHAOSFLESH_API RadialHexMesh(const FVector::FReal InnerRadius, const FVector::FReal OuterRadius, const FVector::FReal Height, const int32 RadialSample, const int32 AngularSample, const int32 VerticalSample, TArray<int32>& HexElements, TArray<FVector>& HexVertices);
+void CHAOSFLESH_API RadialTetMesh(const FVector::FReal InnerRadius, const FVector::FReal OuterRadius, const FVector::FReal Height, const int32 RadialSample, const int32 AngularSample, const int32 VerticalSample, const FVector::FReal BulgeDistance, TArray<FIntVector4>& TetElements, TArray<FVector>& TetVertices);
+void CHAOSFLESH_API RadialHexMesh(const FVector::FReal InnerRadius, const FVector::FReal OuterRadius, const FVector::FReal Height, const int32 RadialSample, const int32 AngularSample, const int32 VerticalSample, const FVector::FReal BulgeDistance, TArray<int32>& HexElements, TArray<FVector>& HexVertices);
 void CHAOSFLESH_API RegularHexMesh2TetMesh(const TArray<FVector>& HexVertices, const TArray<int32>& HexElements, TArray<FVector>& TetVertices, TArray<FIntVector4>& TetElements);
 void CHAOSFLESH_API ComputeHexMeshFaces(const TArray<int32>& HexElements, TArray<FIntVector2>& CommonFaces);
 
diff --git a/Engine/Plugins/Experimental/ChaosFlesh/Source/ChaosFleshNodes/Private/Dataflow/ChaosFleshRadialTetrahedronNodes.cpp b/Engine/Plugins/Experimental/ChaosFlesh/Source/ChaosFleshNodes/Private/Dataflow/ChaosFleshRadialTetrahedronNodes.cpp
--- a/Engine/Plugins/Experimental/ChaosFlesh/Source/ChaosFleshNodes/Private/Dataflow/ChaosFleshRadialTetrahedronNodes.cpp
+++ b/Engine/Plugins/Experimental/ChaosFlesh/Source/ChaosFleshNodes/Private/Dataflow/ChaosFleshRadialTetrahedronNodes.cpp
@@ -43,7 +43,7 @@ void FRadialTetrahedronDataflowNodes::Evaluate(Dataflow::FContext& Context, cons
 		TUniquePtr<FFleshCollection> InCollection(GetValue<DataType>(Context, &Collection).NewCopy<FFleshCollection>());
 		TArray<FIntVector4> TetElements;
 		TArray<FVector> TetVertices;
-		RadialTetMesh(InnerRadius, OuterRadius, Height, RadialSample, AngularSample, VerticalSample, TetElements, TetVertices);
+		RadialTetMesh(InnerRadius, OuterRadius, Height, RadialSample, AngularSample, VerticalSample, BulgeRatio * (OuterRadius - InnerRadius) / (Chaos::FReal)2., TetElements, TetVertices);
 		TArray<FIntVector3> SurfaceElements = Dataflow::GetSurfaceTriangles(TetElements);;
 		TUniquePtr<FTetrahedralCollection> RadialCollection(FTetrahedralCollection::NewTetrahedralCollection(TetVertices, SurfaceElements, TetElements));
 		InCollection->AppendGeometry(*RadialCollection.Get());
diff --git a/Engine/Plugins/Experimental/ChaosFlesh/Source/ChaosFleshNodes/Public/Dataflow/ChaosFleshRadialTetrahedronNodes.h b/Engine/Plugins/Experimental/ChaosFlesh/Source/ChaosFleshNodes/Public/Dataflow/ChaosFleshRadialTetrahedronNodes.h
--- a/Engine/Plugins/Experimental/ChaosFlesh/Source/ChaosFleshNodes/Public/Dataflow/ChaosFleshRadialTetrahedronNodes.h
+++ b/Engine/Plugins/Experimental/ChaosFlesh/Source/ChaosFleshNodes/Public/Dataflow/ChaosFleshRadialTetrahedronNodes.h
@@ -42,6 +42,9 @@ public:
 	UPROPERTY(EditAnywhere, Category = "Dataflow", meta = (ClampMin = "2"))
 	int32 VerticalSample = 2;
 
+	UPROPERTY(EditAnywhere, Category = "Dataflow", meta = (ClampMin = "0.0", ClampMax = "1.0"))
+	double BulgeRatio = double(0.0);
+
 	FRadialTetrahedronDataflowNodes(const Dataflow::FNodeParameters& InParam, FGuid InGuid = FGuid::NewGuid())
 		: FDataflowNode(InParam, InGuid)
 	{
