ChaosFlesh: muscle sim setup (bicep)

diff --git a/Engine/Plugins/Experimental/ChaosFlesh/Source/ChaosFleshNodes/Private/Dataflow/ChaosFleshEngineAssetNodes.cpp b/Engine/Plugins/Experimental/ChaosFlesh/Source/ChaosFleshNodes/Private/Dataflow/ChaosFleshEngineAssetNodes.cpp
--- a/Engine/Plugins/Experimental/ChaosFlesh/Source/ChaosFleshNodes/Private/Dataflow/ChaosFleshEngineAssetNodes.cpp
+++ b/Engine/Plugins/Experimental/ChaosFlesh/Source/ChaosFleshNodes/Private/Dataflow/ChaosFleshEngineAssetNodes.cpp
@@ -24,6 +24,7 @@ namespace Dataflow
 		DATAFLOW_NODE_REGISTER_CREATION_FACTORY(FSetFleshDefaultPropertiesNode);
 		DATAFLOW_NODE_REGISTER_CREATION_FACTORY(FComputeFiberFieldNode);
 		DATAFLOW_NODE_REGISTER_CREATION_FACTORY(FComputeIslandsNode);
+		DATAFLOW_NODE_REGISTER_CREATION_FACTORY(FGenerateOriginInsertionNode);
 	}
 }
 
@@ -217,8 +218,8 @@ void FComputeFiberFieldNode::Evaluate(Dataflow::FContext& Context, const FDatafl
 		//
 
 		FManagedArrayCollection InCollection = GetValue<FManagedArrayCollection>(Context, &Collection);
-		TArray<int32> InOriginIndices;// = GetValue<TArray<int32>>(Context, &OriginIndices);
-		TArray<int32> InInsertionIndices;// = GetValue<TArray<int32>>(Context, &InsertionIndices);
+		TArray<int32> InOriginIndices = GetValue<TArray<int32>>(Context, &OriginIndices);
+		TArray<int32> InInsertionIndices = GetValue<TArray<int32>>(Context, &InsertionIndices);
 
 		// Tetrahedra
 		TManagedArray<FIntVector4>* Elements = InCollection.FindAttribute<FIntVector4>(
@@ -437,3 +438,111 @@ void FComputeIslandsNode::Evaluate(Dataflow::FContext& Context, const FDataflowO
 }
 
 
+void FGenerateOriginInsertionNode::Evaluate(Dataflow::FContext& Context, const FDataflowOutput* Out) const
+{
+	if (Out->IsA<FManagedArrayCollection>(&Collection))
+	{
+		//
+		// Gather inputs
+		//
+
+		FManagedArrayCollection InCollection = GetValue<FManagedArrayCollection>(Context, &Collection);
+		TArray<int32> InOriginIndices = GetValue<TArray<int32>>(Context, &OriginIndicesIn);
+		TArray<int32> InInsertionIndices = GetValue<TArray<int32>>(Context, &InsertionIndicesIn);
+		TArray<int32> OutOriginIndices;
+		TArray<int32> OutInsertionIndices;
+		// Tetrahedra
+		TManagedArray<FIntVector4>* Elements = InCollection.FindAttribute<FIntVector4>(
+			FTetrahedralCollection::TetrahedronAttribute, FTetrahedralCollection::TetrahedralGroup);
+		if (!Elements)
+		{
+			UE_LOG(LogChaosFlesh, Warning,
+				TEXT("GenerateOriginInsertionNode: Failed to find geometry collection attr '%s' in group '%s'"),
+				*FTetrahedralCollection::TetrahedronAttribute.ToString(), *FTetrahedralCollection::TetrahedralGroup.ToString());
+			Out->SetValue<FManagedArrayCollection>(InCollection, Context);
+			return;
+		}
+
+		// Vertices
+		TManagedArray<FVector3f>* Vertex = InCollection.FindAttribute<FVector3f>("Vertex", "Vertices");
+		//TArray<FVector3f>* MeshVertex;
+		if (!Vertex)
+		{
+			UE_LOG(LogChaosFlesh, Warning,
+				TEXT("GenerateOriginInsertionNode: Failed to find geometry collection attr 'Vertex' in group 'Vertices'"));
+			Out->SetValue<FManagedArrayCollection>(InCollection, Context);
+			return;
+		}
+
+		// Incident elements
+		TManagedArray<TArray<int32>>* IncidentElements = InCollection.FindAttribute<TArray<int32>>(
+			FTetrahedralCollection::IncidentElementsAttribute, FGeometryCollection::VerticesGroup);
+		if (!IncidentElements)
+		{
+			UE_LOG(LogChaosFlesh, Warning,
+				TEXT("GenerateOriginInsertionNode: Failed to find geometry collection attr '%s' in group '%s'"),
+				*FTetrahedralCollection::IncidentElementsAttribute.ToString(), *FGeometryCollection::VerticesGroup.ToString());
+			Out->SetValue<FManagedArrayCollection>(InCollection, Context);
+			return;
+		}
+		TManagedArray<TArray<int32>>* IncidentElementsLocalIndex = InCollection.FindAttribute<TArray<int32>>(
+			FTetrahedralCollection::IncidentElementsLocalIndexAttribute, FGeometryCollection::VerticesGroup);
+		if (!IncidentElementsLocalIndex)
+		{
+			UE_LOG(LogChaosFlesh, Warning,
+				TEXT("GenerateOriginInsertionNode: Failed to find geometry collection attr '%s' in group '%s'"),
+				*FTetrahedralCollection::IncidentElementsLocalIndexAttribute.ToString(), *FGeometryCollection::VerticesGroup.ToString());
+			Out->SetValue<FManagedArrayCollection>(InCollection, Context);
+			return;
+		}
+
+		//
+		// Pull Origin & Insertion data out of the geometry collection.  We may want other ways of specifying
+		// these via an input on the node...
+		//
+		auto DoubleVert = [](FVector3f V) { return FVector3d(V.X, V.Y, V.Z); };
+		// Origin vertices
+		if (!InOriginIndices.IsEmpty())
+		{
+			for (int32 i = 0; i < InOriginIndices.Num(); ++i)
+			{
+				if (InOriginIndices[i] < Vertex->Num())
+				{
+					for (int32 j = 0; j < Vertex->Num(); ++j)
+					{
+						if (FVector::Distance(DoubleVert((*Vertex)[InOriginIndices[i]]), DoubleVert((*Vertex)[j])) < Radius)
+						{
+							OutOriginIndices.Add(j);
+						}
+					}
+				}
+			}
+		}
+
+		// Insertion vertices
+		if (!InInsertionIndices.IsEmpty())
+		{
+			for (int32 i = 0; i < InInsertionIndices.Num(); ++i)
+			{
+				if (InInsertionIndices[i] < Vertex->Num())
+				{
+					for (int32 j = 0; j < Vertex->Num(); ++j)
+					{
+						if (FVector::Distance(DoubleVert((*Vertex)[InInsertionIndices[i]]), DoubleVert((*Vertex)[j])) < Radius)
+						{
+							OutInsertionIndices.Add(j);
+						}
+					}
+				}
+			}
+		}
+
+		//
+		// Set output(s)
+		//
+
+		SetValue<FManagedArrayCollection>(Context, InCollection, &Collection);
+		SetValue<TArray<int32>>(Context, OutOriginIndices, &OriginIndicesOut);
+		SetValue<TArray<int32>>(Context, OutInsertionIndices, &InsertionIndicesOut);
+	}
+}
\ No newline at end of file
diff --git a/Engine/Plugins/Experimental/ChaosFlesh/Source/ChaosFleshNodes/Public/Dataflow/ChaosFleshEngineAssetNodes.h b/Engine/Plugins/Experimental/ChaosFlesh/Source/ChaosFleshNodes/Public/Dataflow/ChaosFleshEngineAssetNodes.h
--- a/Engine/Plugins/Experimental/ChaosFlesh/Source/ChaosFleshNodes/Public/Dataflow/ChaosFleshEngineAssetNodes.h
+++ b/Engine/Plugins/Experimental/ChaosFlesh/Source/ChaosFleshNodes/Public/Dataflow/ChaosFleshEngineAssetNodes.h
@@ -121,11 +121,11 @@ public:
 	UPROPERTY(meta = (DataflowInput, DataflowOutput, DisplayName = "Collection", DataflowPassthrough  = "Collection"))
 	FManagedArrayCollection Collection;
 
-	//UPROPERTY(meta = (DataflowInput, DisplayName = "OriginVertexIndices"))
-	//TArray<int32> OriginIndices;
+	UPROPERTY(meta = (DataflowInput, DisplayName = "OriginVertexIndices"))
+	TArray<int32> OriginIndices;
 
-	//UPROPERTY(meta = (DataflowInput, DisplayName = "InsertionVertexIndices"))
-	//TArray<int32> InsertionIndices;
+	UPROPERTY(meta = (DataflowInput, DisplayName = "InsertionVertexIndices"))
+	TArray<int32> InsertionIndices;
 
 	UPROPERTY(EditAnywhere, Category = "Dataflow")
 	FString OriginInsertionGroupName = FString();
@@ -146,6 +146,8 @@ public:
 		: FDataflowNode(InParam, InGuid)
 	{
 		RegisterInputConnection(&Collection);
+		RegisterInputConnection(&OriginIndices);
+		RegisterInputConnection(&InsertionIndices);
 		RegisterOutputConnection(&Collection, &Collection);
 	}
 
@@ -186,6 +188,49 @@ public:
 	virtual void Evaluate(Dataflow::FContext& Context, const FDataflowOutput* Out) const override;
 };
 
+//Given two sets of vertex indices, generate two sets of vertex indices for origins and insertions that are within X distance away.
+USTRUCT(meta = (DataflowFlesh))
+struct FGenerateOriginInsertionNode : public FDataflowNode
+{
+	GENERATED_USTRUCT_BODY()
+	DATAFLOW_NODE_DEFINE_INTERNAL(FGenerateOriginInsertionNode, "GenerateOriginInsertion", "Flesh", "")
+	DATAFLOW_NODE_RENDER_TYPE(FGeometryCollection::StaticType(), "Collection")
+
+public:
+	//typedef FManagedArrayCollection DataType;
+
+	UPROPERTY(meta = (DataflowInput, DataflowOutput, DisplayName = "Collection", DataflowPassthrough = "Collection"))
+		FManagedArrayCollection Collection;
+
+	UPROPERTY(meta = (DataflowInput, DisplayName = "OriginIndicesIn"))
+		TArray<int32> OriginIndicesIn;
+
+	UPROPERTY(meta = (DataflowInput, DisplayName = "InsertionIndicesIn"))
+		TArray<int32> InsertionIndicesIn;
+
+	UPROPERTY(meta = (DataflowOutput, DisplayName = "OriginIndicesOut"))
+		TArray<int32> OriginIndicesOut;
+
+	UPROPERTY(meta = (DataflowOutput, DisplayName = "InsertionIndicesOut"))
+		TArray<int32> InsertionIndicesOut;
+	
+	UPROPERTY(EditAnywhere, Category = "Dataflow")
+		float Radius = float(1);
+
+	FGenerateOriginInsertionNode(const Dataflow::FNodeParameters& InParam, FGuid InGuid = FGuid::NewGuid())
+		: FDataflowNode(InParam, InGuid)
+	{
+		RegisterInputConnection(&Collection);
+		RegisterInputConnection(&OriginIndicesIn);
+		RegisterInputConnection(&InsertionIndicesIn);
+		RegisterOutputConnection(&Collection, &Collection);
+		RegisterOutputConnection(&OriginIndicesOut);
+		RegisterOutputConnection(&InsertionIndicesOut);
+	}
+
+	virtual void Evaluate(Dataflow::FContext& Context, const FDataflowOutput* Out) const override;
+};
+
 namespace Dataflow
 {
 	void RegisterChaosFleshEngineAssetNodes();
