ChaosFlesh : Transient Kinematic Particles

diff --git a/Engine/Plugins/Experimental/ChaosFlesh/Source/ChaosFleshEngine/Private/ChaosFlesh/ChaosDeformableGameplayComponent.cpp b/Engine/Plugins/Experimental/ChaosFlesh/Source/ChaosFleshEngine/Private/ChaosFlesh/ChaosDeformableGameplayComponent.cpp
--- a/Engine/Plugins/Experimental/ChaosFlesh/Source/ChaosFleshEngine/Private/ChaosFlesh/ChaosDeformableGameplayComponent.cpp
+++ b/Engine/Plugins/Experimental/ChaosFlesh/Source/ChaosFleshEngine/Private/ChaosFlesh/ChaosDeformableGameplayComponent.cpp
@@ -5,10 +5,28 @@
 #include "ChaosStats.h"
 #include "Components/BoxComponent.h"
 #include "Chaos/DebugDrawQueue.h"
-
+#include "ChaosFlesh/ChaosDeformableTypes.h"
+#include "Components/ActorComponent.h"
+#include "Components/SkeletalMeshComponent.h"
+#include "Containers/Map.h"
+#include "Dataflow/DataflowEngineUtil.h"
+#include "GeometryCollection/Facades/CollectionConstraintOverrideFacade.h"
+#include "GeometryCollection/Facades/CollectionTransformSourceFacade.h"
+#include "GeometryCollection/Facades/CollectionVertexBoneWeightsFacade.h"
+#include "HAL/IConsoleManager.h"
+#include "WorldCollision.h"
+#include "GeometryCollection/GeometryCollectionAlgo.h"
+#include "Engine/World.h"
+#include "Animation/Skeleton.h"
+#include "Animation/Rig.h"
+#include "Engine/SkeletalMesh.h"
 
 #include UE_INLINE_GENERATED_CPP_BY_NAME(ChaosDeformableGameplayComponent)
+#define LOCTEXT_NAMESPACE "DeformableGameplayComponent"
 
+FChaosEngineDeformableCVarParams CVarParams2;
+FAutoConsoleVariableRef CVarDeformableDebugParamsDrawSceneRaycasts(TEXT("p.Chaos.DebugDraw.Deformable.SceneRaycasts"), CVarParams2.bDoDrawSceneRaycasts, TEXT("Debug draw the deformables scene raycasts. [def: false]"));
+FAutoConsoleVariableRef CVarDeformableEnvCollisionsLineTracesBatchSize(TEXT("p.Chaos.Deformable.EnvCollisionsLineTracesBatchSize"), CVarParams2.EnvCollisionsLineTraceBatchSize, TEXT("Batch size for FleshComponent env collsions. [def: 10]"));
 
 //DECLARE_CYCLE_STAT(TEXT("Chaos.Deformable.UDeformableGameplayComponent.StreamStaticMeshGeometry"), STAT_ChaosDeformable_UDeformableGameplayComponent_StreamStaticMeshGeometry, STATGROUP_Chaos);
 
@@ -21,3 +39,303 @@ UDeformableGameplayComponent::UDeformableGameplayComponent(const FObjectInitiali
 UDeformableGameplayComponent::~UDeformableGameplayComponent()
 {
 }
+
+
+void UDeformableGameplayComponent::PreSolverUpdate()
+{
+	if (RigBoundRayCasts.bEnableRigBoundRaycasts)
+	{
+		DetectEnvironmentCollisions(RigBoundRayCasts.MaxNumTests, RigBoundRayCasts.bTestDownOnly, RigBoundRayCasts.TestRange, RigBoundRayCasts.CollisionChannel.GetValue());
+	}
+}
+
+
+void UDeformableGameplayComponent::DetectEnvironmentCollisions(const int32 MaxNumTests, const bool bTestDownOnly, const float TestRange, const ECollisionChannel CollisionChannel)
+{
+	auto ToDouble = [](FVector3f V) { return FVector(V[0], V[1], V[2]); };
+	auto ToSingle = [](FVector V) { return FVector3f(static_cast<float>(V[0]), static_cast<float>(V[1]), static_cast<float>(V[2])); };
+
+	if (!GetRestCollection() || !GetSimulationCollection() || !GetDynamicCollection())
+	{
+		return;
+	}
+
+	// Init outputs, and clear any existing.
+	GeometryCollection::Facades::FConstraintOverrideTargetFacade CnstrOutputs(
+		*GetSimulationCollection()->GetCollection());
+	CnstrOutputs.DefineSchema();
+	CnstrOutputs.Clear();
+
+	// Init inputs, and return if we have no candidates.
+	GeometryCollection::Facades::FConstraintOverrideCandidateFacade CnstrCandidates(
+		*GetRestCollection()->GetCollection());
+	if (!CnstrCandidates.IsValid())
+	{
+		return;
+	}
+	if (CnstrCandidates.Num() == 0)
+	{
+		return;
+	}
+
+	// We need to test the current skinned positions of the rest collection against the 
+	// environment.
+	const TManagedArray<FVector3f>* Positions = GetRestCollection()->FindPositions();
+	if (!Positions)
+	{
+		return;
+	}
+
+	// Get bone weights for skinning vertices of the tetrahedron mesh.
+	GeometryCollection::Facades::FVertexBoneWeightsFacade BoneWeightsFacade(
+		*GetRestCollection()->GetCollection());
+	if (!BoneWeightsFacade.IsValid())
+	{
+		return;
+	}
+	const TManagedArray<TArray<int32>>* BoneIndices = BoneWeightsFacade.FindBoneIndices();
+	const TManagedArray<TArray<float>>* BoneWeights = BoneWeightsFacade.FindBoneWeights();
+
+	const TManagedArray<FTransform>* RestTransforms = nullptr;
+	if (const UFleshAsset* FleshAsset = GetRestCollection())
+	{
+		if (const FFleshCollection* Rest = FleshAsset->GetCollection())
+		{
+			GeometryCollection::Facades::FTransformSource TransformSource(*Rest);
+			if (TransformSource.IsValid())
+			{
+				RestTransforms = Rest->FindAttribute<FTransform>(
+					FTransformCollection::TransformAttribute, FTransformCollection::TransformGroup);
+			}
+		}
+	}
+
+	// Find the skeletal mesh containing the bone we're using for the origin of our rays.
+	USkeletalMeshComponent* SkeletalMeshComponent = nullptr;
+	TArray<USkeletalMeshComponent*> SkeletalMeshComponents;
+	GetOwner()->GetComponents<USkeletalMeshComponent>(SkeletalMeshComponents);
+	for (int32 i = 0; i < SkeletalMeshComponents.Num(); i++)
+	{
+		if (SkeletalMeshComponents[i]->GetSkeletalMeshAsset() == GetRestCollection()->SkeletalMesh)
+		{
+			SkeletalMeshComponent = SkeletalMeshComponents[i];
+			break;
+		}
+	}
+	if (!SkeletalMeshComponent)
+	{
+		return;
+	}
+
+	// Get the current animated bone transforms and then their positions.
+	const TArray<FTransform>& AnimationInComponentSpaceTransforms = SkeletalMeshComponent->GetComponentSpaceTransforms();
+	const TArray<FTransform>& RestInLocalSpaceTransforms = SkeletalMeshComponent->GetSkeletalMeshAsset()->GetSkeleton()->GetReferenceSkeleton().GetRefBonePose(); // Local Space
+
+	TMap<int32, FVector> WorldBoneLocations;
+	WorldBoneLocations.Reserve(16);
+
+	struct CandidateData
+	{
+		CandidateData(const int32 IndexIn, const FVector& OriginIn, const FVector& TargetIn)
+			: Index(IndexIn), Origin(OriginIn), Target(TargetIn)
+		{}
+		int32 Index;
+		FVector Origin;
+		FVector Target;
+	};
+	TArray<CandidateData> Candidates;
+
+	const FTransform& ComponentToWorldXf = GetComponentToWorld();
+	for (int32 i = 0; i < CnstrCandidates.Num(); i++)
+	{
+		GeometryCollection::Facades::FConstraintOverridesCandidateData Candidate = CnstrCandidates.Get(i);
+		const int32 TetVertexIndex = Candidate.VertexIndex;
+		const int32 OriginBoneIndex = Candidate.BoneIndex;
+
+		if (TetVertexIndex < 0 || TetVertexIndex >= Positions->Num())
+		{
+			continue;
+		}
+		const FVector3f& RestPos = (*Positions)[TetVertexIndex];
+
+		// Compute skinning, if present.
+		bool bDidSkinnedPosition = false;
+		FVector SkinnedPos(0.0,0.0,0.0);
+		if (RestTransforms && RestTransforms->Num() &&
+			BoneIndices && TetVertexIndex < BoneIndices->Num() && (*BoneIndices)[TetVertexIndex].Num() &&
+			BoneWeights && TetVertexIndex < BoneWeights->Num() && (*BoneWeights)[TetVertexIndex].Num())
+		{
+			// Get Skinned bones and weights.
+			const TArray<int32>& Bones = (*BoneIndices)[TetVertexIndex];
+			const TArray<float>& Weights = (*BoneWeights)[TetVertexIndex];
+			checkSlow(Bones.Num() == Weights.Num());
+
+			// Compute skinned test mesh world position.
+			
+			for (int32 j = 0; j < Bones.Num(); j++)
+			{
+				if (AnimationInComponentSpaceTransforms.IsValidIndex(Bones[j]) && RestTransforms->Num() > Bones[j])
+				{
+					//auto ComponentSpaceMatrix = [](const FReferenceSkeleton& RefSkeleton, const TArray<FTransform>& RestTransforms, int32 Index)
+					//{
+					//	FTransform FlatRest = RestTransforms[Index];
+					//	while (RefSkeleton.GetParentIndex(Index) != INDEX_NONE)
+					//	{
+					//		Index = RefSkeleton.GetParentIndex(Index);
+					//		FlatRest *= RestTransforms[Index];
+					//	}
+					//	return FlatRest;
+					//};
+					//
+					//const FReferenceSkeleton& RefSkeleton = SkeletalMeshComponent->GetSkeletalMeshAsset()->GetSkeleton()->GetReferenceSkeleton();
+					//const TArray<FTransform>& RefMat = RefSkeleton.GetRefBonePose();
+					//FTransform RestInComponent = ComponentSpaceMatrix(SkeletalMeshComponent->GetSkeletalMeshAsset()->GetSkeleton()->GetReferenceSkeleton(), RefMat, Bones[j]);
+					//FVector LocalPoint = RestInComponent.TransformPosition(ToDouble(RestPos));
+					
+					const FFleshCollection* Collection = GetRestCollection()->GetCollection();
+					FVector LocalPoint = (*RestTransforms)[Bones[j]].InverseTransformPosition(ToDouble(RestPos));
+					if (const TManagedArray<FTransform>* ConstTransformsPtr = Collection->FindAttributeTyped<FTransform>(FName("ComponentTransform"), FName("ComponentTransformGroup")))
+					{
+						LocalPoint = (*ConstTransformsPtr)[0].InverseTransformPosition(ToDouble(RestPos));
+					}
+
+					SkinnedPos += AnimationInComponentSpaceTransforms[Bones[j]].TransformPosition(LocalPoint) * Weights[j];
+					bDidSkinnedPosition = true;
+				}
+			}
+		}
+		if (!bDidSkinnedPosition)
+		{
+			SkinnedPos = FVector(RestPos[0], RestPos[1], RestPos[2]);
+		}
+		FVector WorldSkinnedPos = ComponentToWorldXf.TransformPosition(SkinnedPos);
+
+		// Compute/lookup ray origin world position.
+		FVector WorldOrigin;
+		if (FVector* Existing = WorldBoneLocations.Find(OriginBoneIndex))
+		{
+			WorldOrigin = *Existing;
+		}
+		else
+		{
+			if (!AnimationInComponentSpaceTransforms.IsValidIndex(OriginBoneIndex))
+			{
+				continue;
+			}
+			const FTransform& BoneXf = AnimationInComponentSpaceTransforms[OriginBoneIndex];
+			WorldOrigin = ComponentToWorldXf.TransformPosition(BoneXf.GetLocation());
+			WorldBoneLocations.Add(TTuple<int32, FVector>(OriginBoneIndex, WorldOrigin));
+		}
+
+		if (bTestDownOnly)
+		{
+			FVector Direction = WorldSkinnedPos - WorldOrigin;
+			Direction.Normalize();
+			if (FVector::DotProduct(Direction, FVector::DownVector) < TestRange) // +1 down, -1 up, 0 orthogonal - DownVector:(0,0,-1)
+			{
+				continue;
+			}
+		}
+
+		// Queue the query.
+		Candidates.Add(CandidateData(TetVertexIndex, WorldOrigin, WorldSkinnedPos));
+	}
+	if (!Candidates.Num())
+	{
+		return;
+	}
+
+	// Randomly discard potentials until we get under MaxNumTests.
+	if (EnvCollisionsPreviousHits.Num() > MaxNumTests)
+	{
+		// MaxNumTests changed, which puts us in danger of never exiting the while loop.
+		// Blow the buffer.
+		EnvCollisionsPreviousHits.Reset();
+	}
+	while (Candidates.Num() > MaxNumTests)
+	{
+		const int32 Index = static_cast<int32>(FGenericPlatformMath::Rand() / RAND_MAX * (Candidates.Num() - 1));
+		if (!EnvCollisionsPreviousHits.Contains(Index))
+		{
+			Candidates.RemoveAtSwap(Index);
+		}
+	}
+	EnvCollisionsPreviousHits.Reset();
+
+	// Run scene queries.
+	const FString QueryName =
+		FText::Format(
+		LOCTEXT("DetectEnvironmentCollisions", "FleshComponent: '{0}'"),
+		FText::FromString(GetName())).ToString();
+	FCollisionQueryParams TraceParams(
+		FName(*QueryName),
+		false,			// false to ignore complex collisions 
+		GetOwner());	// GetOwner() to ignore self
+	const FTransform WorldToComponentXf(ComponentToWorldXf.ToMatrixWithScale().Inverse());
+
+	TArray<FHitResult> HitBuffer;
+	HitBuffer.SetNum(Candidates.Num());
+	TArray<bool> HitResult;
+	HitResult.SetNum(Candidates.Num());
+
+	Chaos::PhysicsParallelForRange(Candidates.Num(), [&](const int32 StartIndex, const int32 StopIndex)
+								   {
+									   for (int32 Index = StartIndex; Index < StopIndex; ++Index)
+									   {
+										   const int32 TetVertexIndex = Candidates[Index].Index;
+										   const FVector& WorldOrigin = Candidates[Index].Origin;
+										   const FVector& WorldSkinnedPos = Candidates[Index].Target;
+										   HitResult[Index] =
+											   GetWorld()->LineTraceSingleByObjectType(
+											   HitBuffer[Index],
+											   WorldOrigin,
+											   WorldSkinnedPos,
+											   CollisionChannel,
+											   TraceParams);
+									   }
+								   },
+								   CVarParams2.EnvCollisionsLineTraceBatchSize,
+									   Candidates.Num() < CVarParams2.EnvCollisionsLineTraceBatchSize);
+
+	// Process results.
+	for (int32 i = 0; i < Candidates.Num(); i++)
+	{
+		const int32 TetVertexIndex = Candidates[i].Index;
+		const FVector& WorldOrigin = Candidates[i].Origin;
+		const FVector& WorldSkinnedPos = Candidates[i].Target;
+		if (HitResult[i])
+		{
+			FHitResult& Hit = HitBuffer[i];
+
+			if (EnvironmentCollisionsSkipList.Contains(Hit.Component))
+			{
+				continue;
+			}
+			GeometryCollection::Facades::FConstraintOverridesTargetData CnstrData;
+			CnstrData.VertexIndex = TetVertexIndex;
+			CnstrData.PositionTarget = ToSingle(WorldToComponentXf.TransformPosition(Hit.ImpactPoint)); // Component space
+			CnstrOutputs.Add(CnstrData);
+
+			EnvCollisionsPreviousHits.Add(TetVertexIndex);
+
+#if WITH_EDITOR
+			if (CVarParams2.IsDebugDrawingEnabled() && CVarParams2.bDoDrawSceneRaycasts)
+			{
+				Chaos::FDebugDrawQueue::GetInstance().DrawDebugDirectionalArrow(WorldOrigin, WorldSkinnedPos, 2.0f, FColor::Red, false, -1.0f, 0, 0.75f);
+				Chaos::FDebugDrawQueue::GetInstance().DrawDebugPoint(Hit.ImpactPoint, FColor::Yellow, false, -1.0f, 0, 5);
+			}
+#endif
+		}
+		else
+		{
+#if WITH_EDITOR
+			if (CVarParams2.IsDebugDrawingEnabled() && CVarParams2.bDoDrawSceneRaycasts)
+			{
+				Chaos::FDebugDrawQueue::GetInstance().DrawDebugDirectionalArrow(WorldOrigin, WorldSkinnedPos, 2.0f, FColor::Green, false, -1.0f, 0, 0.5f);
+			}
+#endif
+		}
+	}
+}
+
+#undef LOCTEXT_NAMESPACE
diff --git a/Engine/Plugins/Experimental/ChaosFlesh/Source/ChaosFleshEngine/Private/ChaosFlesh/ChaosDeformableSolverComponent.cpp b/Engine/Plugins/Experimental/ChaosFlesh/Source/ChaosFleshEngine/Private/ChaosFlesh/ChaosDeformableSolverComponent.cpp
--- a/Engine/Plugins/Experimental/ChaosFlesh/Source/ChaosFleshEngine/Private/ChaosFlesh/ChaosDeformableSolverComponent.cpp
+++ b/Engine/Plugins/Experimental/ChaosFlesh/Source/ChaosFleshEngine/Private/ChaosFlesh/ChaosDeformableSolverComponent.cpp
@@ -264,6 +264,7 @@ void UDeformableSolverComponent::UpdateFromGameThread(float DeltaTime)
 			{
 				if (IsSimulating(DeformableComponent))
 				{
+					DeformableComponent->PreSolverUpdate();
 					if (FDataMapValue Value = DeformableComponent->NewDeformableData())
 					{
 						DataMap.Add(DeformableComponent, Value);
diff --git a/Engine/Plugins/Experimental/ChaosFlesh/Source/ChaosFleshEngine/Private/ChaosFlesh/FleshAsset.cpp b/Engine/Plugins/Experimental/ChaosFlesh/Source/ChaosFleshEngine/Private/ChaosFlesh/FleshAsset.cpp
--- a/Engine/Plugins/Experimental/ChaosFlesh/Source/ChaosFleshEngine/Private/ChaosFlesh/FleshAsset.cpp
+++ b/Engine/Plugins/Experimental/ChaosFlesh/Source/ChaosFleshEngine/Private/ChaosFlesh/FleshAsset.cpp
@@ -50,6 +50,16 @@ void UFleshAsset::PostEditCallback()
 	//UE_LOG(LogFleshAssetInternal, Log, TEXT("UFleshAsset::PostEditCallback()"));
 }
 
+TManagedArray<FVector3f>& UFleshAsset::GetPositions()
+{
+	return FleshCollection->ModifyAttribute<FVector3f>("Vertex", FGeometryCollection::VerticesGroup);
+}
+
+const TManagedArray<FVector3f>* UFleshAsset::FindPositions() const
+{
+	return FleshCollection->FindAttributeTyped<FVector3f>("Vertex", FGeometryCollection::VerticesGroup);
+}
+
 /** Serialize */
 void UFleshAsset::Serialize(FArchive& Ar)
 {
diff --git a/Engine/Plugins/Experimental/ChaosFlesh/Source/ChaosFleshEngine/Private/ChaosFlesh/FleshComponent.cpp b/Engine/Plugins/Experimental/ChaosFlesh/Source/ChaosFleshEngine/Private/ChaosFlesh/FleshComponent.cpp
--- a/Engine/Plugins/Experimental/ChaosFlesh/Source/ChaosFleshEngine/Private/ChaosFlesh/FleshComponent.cpp
+++ b/Engine/Plugins/Experimental/ChaosFlesh/Source/ChaosFleshEngine/Private/ChaosFlesh/FleshComponent.cpp
@@ -29,6 +29,7 @@ FAutoConsoleVariableRef CVarDeforambleDoDrawSimulationMesh(TEXT("p.Chaos.DebugDr
 FAutoConsoleVariableRef CVarDeforambleDoDrawSkeletalMeshBindingPositions(TEXT("p.Chaos.DebugDraw.Deformable.SkeletalMeshBindingPositions"), CVarParams.bDoDrawSkeletalMeshBindingPositions, TEXT("Debug draw the deformable simulation's SkeletalMeshBindingPositions on the game thread. [def: false]"));
 FAutoConsoleVariableRef CVarDeforambleDoDrawSkeletalMeshBindingPositionsSimulationBlendWeight(TEXT("p.Chaos.DebugDraw.Deformable.SkeletalMeshBindingPositions.SimulationBlendWeight"), CVarParams.DrawSkeletalMeshBindingPositionsSimulationBlendWeight, TEXT("Set the simulation blend weight of the skeletal mesh debug draw.[def: 1.]"));
 
+
 #define PERF_SCOPE(X) SCOPE_CYCLE_COUNTER(X); TRACE_CPUPROFILER_EVENT_SCOPE(X);
 DECLARE_CYCLE_STAT(TEXT("Chaos.Deformable.UFleshComponent.TickComponent"), STAT_ChaosDeformable_UFleshComponent_TickComponent, STATGROUP_Chaos);
 DECLARE_CYCLE_STAT(TEXT("Chaos.Deformable.UFleshComponent.NewProxy"), STAT_ChaosDeformable_UFleshComponent_NewProxy, STATGROUP_Chaos);
@@ -47,6 +48,7 @@ UFleshComponent::UFleshComponent(const FObjectInitializer& ObjectInitializer)
 	bTickInEditor = CVarParams.bDoDrawSimulationMesh;
 
 	DynamicCollection = ObjectInitializer.CreateDefaultSubobject<UFleshDynamicAsset>(this, TEXT("Flesh Dynamic Asset"));
+	SimulationCollection = ObjectInitializer.CreateDefaultSubobject<USimulationAsset>(this, TEXT("Flesh Simulation Asset"));
 }
 
 UFleshComponent::~UFleshComponent()
@@ -75,6 +77,10 @@ void UFleshComponent::EndPlay(const EEndPlayReason::Type ReasonEnd)
 	{
 		GetDynamicCollection()->Reset();
 	}
+	if (GetSimulationCollection())
+	{
+		GetSimulationCollection()->Reset();
+	}
 }
 
 void UFleshComponent::SetRestCollection(const UFleshAsset* InRestCollection)
@@ -101,6 +107,10 @@ UDeformablePhysicsComponent::FThreadingProxy* UFleshComponent::NewProxy()
 				{
 					DynamicCollection = NewObject<UFleshDynamicAsset>(this, TEXT("Flesh Dynamic Asset"));
 				}
+				if (!GetSimulationCollection())
+				{
+					SimulationCollection = NewObject<USimulationAsset>(this, TEXT("Flesh Simulation Asset"));
+				}
 
 				GetDynamicCollection()->Reset(Rest);
 				if (const FManagedArrayCollection* Dynamic = GetDynamicCollection()->GetCollection())
@@ -194,7 +204,8 @@ UDeformablePhysicsComponent::FDataMapValue UFleshComponent::NewDeformableData()
 
 						return FDataMapValue(
 							new Chaos::Softs::FFleshThreadingProxy::FFleshInputBuffer(
-								this->GetComponentTransform(), 
+								*GetSimulationCollection()->GetCollection(),
+								this->GetComponentTransform(),
 								BoneSpaceXf, 
 								SimSpaceTransformGlobalIndex,
 								AnimationTransforms, 
@@ -213,7 +224,8 @@ UDeformablePhysicsComponent::FDataMapValue UFleshComponent::NewDeformableData()
 	}
 	return FDataMapValue(
 		new Chaos::Softs::FFleshThreadingProxy::FFleshInputBuffer(
-			this->GetComponentTransform(), 
+			*GetSimulationCollection()->GetCollection(),
+			this->GetComponentTransform(),
 			GetSimSpaceRestTransform(),
 			SimSpaceTransformGlobalIndex,
 			bTempEnableGravity,
@@ -927,3 +939,5 @@ TArray<FVector> UFleshComponent::GetSkeletalMeshBindingPositionsInternal(const U
 	}
 	return TransformPositions;
 }
+
+
diff --git a/Engine/Plugins/Experimental/ChaosFlesh/Source/ChaosFleshEngine/Private/ChaosFlesh/SimulationAsset.cpp b/Engine/Plugins/Experimental/ChaosFlesh/Source/ChaosFleshEngine/Private/ChaosFlesh/SimulationAsset.cpp
new file mode 100644
--- /dev/null
+++ b/Engine/Plugins/Experimental/ChaosFlesh/Source/ChaosFleshEngine/Private/ChaosFlesh/SimulationAsset.cpp
@@ -0,0 +1,50 @@
+// Copyright Epic Games, Inc. All Rights Reserved.
+
+/*=============================================================================
+	FleshAsset.cpp: UFleshAsset methods.
+=============================================================================*/
+#include "ChaosFlesh/SimulationAsset.h"
+#include "GeometryCollection/GeometryCollection.h"
+
+#include UE_INLINE_GENERATED_CPP_BY_NAME(SimulationAsset)
+
+DEFINE_LOG_CATEGORY_STATIC(LogSimulationAssetInternal, Log, All);
+
+
+USimulationAsset::USimulationAsset(const FObjectInitializer& ObjectInitializer)
+	: Super(ObjectInitializer)
+	, SimulationCollection(new FManagedArrayCollection())
+{
+	Init();
+}
+
+void USimulationAsset::Init()
+{
+	SimulationCollection->AddAttribute<FVector3f>("ObjectState", FGeometryCollection::VerticesGroup);
+}
+
+void USimulationAsset::Reset(const FManagedArrayCollection* InCopyFrom)
+{
+	SimulationCollection = TSharedPtr<FManagedArrayCollection, ESPMode::ThreadSafe>( new FManagedArrayCollection());
+	Init();
+
+	if (InCopyFrom)
+	{
+		GetObjectState().Fill((int32)Chaos::EObjectStateType::Dynamic);
+		SimulationCollection->CopyMatchingAttributesFrom(*InCopyFrom);
+	}
+}
+
+void USimulationAsset::ResetAttributesFrom(const FManagedArrayCollection* InCopyFrom)
+{
+	if (InCopyFrom)
+	{
+		SimulationCollection->CopyMatchingAttributesFrom(*InCopyFrom);
+	}
+}
+
+TManagedArray<int32>& USimulationAsset::GetObjectState()
+{
+	return SimulationCollection->ModifyAttribute<int32>("ObjectState", FGeometryCollection::VerticesGroup);
+}
+
diff --git a/Engine/Plugins/Experimental/ChaosFlesh/Source/ChaosFleshEngine/Public/ChaosFlesh/ChaosDeformableGameplayComponent.h b/Engine/Plugins/Experimental/ChaosFlesh/Source/ChaosFleshEngine/Public/ChaosFlesh/ChaosDeformableGameplayComponent.h
--- a/Engine/Plugins/Experimental/ChaosFlesh/Source/ChaosFleshEngine/Public/ChaosFlesh/ChaosDeformableGameplayComponent.h
+++ b/Engine/Plugins/Experimental/ChaosFlesh/Source/ChaosFleshEngine/Public/ChaosFlesh/ChaosDeformableGameplayComponent.h
@@ -2,6 +2,7 @@
 #pragma once
 
 #include "CoreMinimal.h"
+
 #include "ChaosFlesh/FleshComponent.h"
 #include "UObject/ObjectMacros.h"
 
@@ -11,6 +12,27 @@
 class UBoxComponent;
 class UStaticMeshComponent;
 
+USTRUCT()
+struct FRigBoundRayCasts
+{
+	GENERATED_USTRUCT_BODY()
+
+	UPROPERTY(EditAnywhere, Category = "RigBoundRaycasts")
+		bool bEnableRigBoundRaycasts = false;
+
+	UPROPERTY(EditAnywhere, Category = "RigBoundRaycasts")
+		int32 MaxNumTests = 5;
+
+	UPROPERTY(EditAnywhere, Category = "RigBoundRaycasts")
+		bool bTestDownOnly = false;
+
+	UPROPERTY(EditAnywhere, Category = "RigBoundRaycasts")
+		float TestRange = 0.5f;
+
+	UPROPERTY(EditAnywhere, Category = "RigBoundRaycasts")
+		TEnumAsByte<ECollisionChannel> CollisionChannel = ECC_WorldStatic;
+};
+
 /**
 *	UDeformableGameplayComponent
 */
@@ -22,5 +44,37 @@ class CHAOSFLESHENGINE_API UDeformableGameplayComponent : public UFleshComponent
 public:
 	~UDeformableGameplayComponent();
 
+
+	virtual void PreSolverUpdate() override;
+
+
+	/** 
+	*  DetectEnvironmentCollisions
+	* 
+	* Do raycasts against the scene and feed the results into transient constraints in the solver.
+	* Setup data for this method is initialized in the FAuthorSceneCollisionCandidates dataflow node,
+	* stored in the rest collection. The TargetDeformationSkeleton on the rest collection is used
+	* for the origin of the raycasts.
+	*
+	* \p MaxNumTests is the upper limit on the number of scene line tests.
+	* \p bTestDownOnly culls line tests that are not in the "down" direction.
+	* \p TestRange indicates how close to "down" tests need to be; 1 means only vectors pointing
+	*    straight down are tested, 0.5 would include any vector about 45 degrees from straight down,
+	*    0.0 means any vector pointing down at all, and negative values up to -1 would include vectors
+	*    pointing up.
+	* \p CollisionChannel denotes which scene objects to collide against.
+	*/
+	void DetectEnvironmentCollisions(const int32 MaxNumTests = 100, const bool bTestDownOnly = true, const float TestRange = 0.0, const ECollisionChannel CollisionChannel = ECollisionChannel::ECC_PhysicsBody);
+
+
+	UPROPERTY(EditAnywhere, Category = "ChaosDeformable")
+	FRigBoundRayCasts RigBoundRayCasts;
+
+	/** Objects to skip if hit during \c DetectEnvironmentCollisions(). */
+	UPROPERTY(EditAnywhere, Category = "ChaosDeformable")
+	TArray<TWeakObjectPtr<UPrimitiveComponent>> EnvironmentCollisionsSkipList;
+
+private:
+	TSet<int32> EnvCollisionsPreviousHits;
 };
 
diff --git a/Engine/Plugins/Experimental/ChaosFlesh/Source/ChaosFleshEngine/Public/ChaosFlesh/ChaosDeformablePhysicsComponent.h b/Engine/Plugins/Experimental/ChaosFlesh/Source/ChaosFleshEngine/Public/ChaosFlesh/ChaosDeformablePhysicsComponent.h
--- a/Engine/Plugins/Experimental/ChaosFlesh/Source/ChaosFleshEngine/Public/ChaosFlesh/ChaosDeformablePhysicsComponent.h
+++ b/Engine/Plugins/Experimental/ChaosFlesh/Source/ChaosFleshEngine/Public/ChaosFlesh/ChaosDeformablePhysicsComponent.h
@@ -36,6 +36,7 @@ public:
 	virtual FThreadingProxy* NewProxy() { return nullptr; }
 	virtual void AddProxy(Chaos::Softs::FDeformableSolver::FGameThreadAccess& GameThreadSolver);
 	virtual void RemoveProxy(Chaos::Softs::FDeformableSolver::FGameThreadAccess& GameThreadSolver);
+	virtual void PreSolverUpdate() {};
 	virtual FDataMapValue NewDeformableData() { return FDataMapValue(nullptr); }
 	virtual void UpdateFromSimualtion(const FDataMapValue* SimualtionBuffer) {}
 
diff --git a/Engine/Plugins/Experimental/ChaosFlesh/Source/ChaosFleshEngine/Public/ChaosFlesh/ChaosDeformableTypes.h b/Engine/Plugins/Experimental/ChaosFlesh/Source/ChaosFleshEngine/Public/ChaosFlesh/ChaosDeformableTypes.h
--- a/Engine/Plugins/Experimental/ChaosFlesh/Source/ChaosFleshEngine/Public/ChaosFlesh/ChaosDeformableTypes.h
+++ b/Engine/Plugins/Experimental/ChaosFlesh/Source/ChaosFleshEngine/Public/ChaosFlesh/ChaosDeformableTypes.h
@@ -27,4 +27,17 @@ struct FChaosEngineDeformableCVarParams
 #endif
 	bool bDoDrawSkeletalMeshBindingPositions = false;
 	float DrawSkeletalMeshBindingPositionsSimulationBlendWeight = 1.f;
+	bool bDoDrawSceneRaycasts = false;
+
+	int32 EnvCollisionsLineTraceBatchSize = 10;
+
+	bool IsDebugDrawingEnabled()
+	{
+#if WITH_EDITOR
+		// p.Chaos.DebugDraw.Enabled 1
+		return Chaos::FDebugDrawQueue::GetInstance().IsDebugDrawingEnabled();
+#else
+		return false;
+#endif
+	}
 };
\ No newline at end of file
diff --git a/Engine/Plugins/Experimental/ChaosFlesh/Source/ChaosFleshEngine/Public/ChaosFlesh/FleshAsset.h b/Engine/Plugins/Experimental/ChaosFlesh/Source/ChaosFleshEngine/Public/ChaosFlesh/FleshAsset.h
--- a/Engine/Plugins/Experimental/ChaosFlesh/Source/ChaosFleshEngine/Public/ChaosFlesh/FleshAsset.h
+++ b/Engine/Plugins/Experimental/ChaosFlesh/Source/ChaosFleshEngine/Public/ChaosFlesh/FleshAsset.h
@@ -85,6 +85,9 @@ public:
 	const FFleshCollection* GetCollection() const { return FleshCollection.Get(); }
 	FFleshCollection* GetCollection() { return FleshCollection.Get(); }
 
+	TManagedArray<FVector3f>& GetPositions();
+	const TManagedArray<FVector3f>* FindPositions() const;
+
 	FFleshAssetEdit EditCollection() const {
 		UFleshAsset* ThisNC = const_cast<UFleshAsset*>(this); 
 		return FFleshAssetEdit(ThisNC, [ThisNC]() {ThisNC->PostEditCallback(); });
diff --git a/Engine/Plugins/Experimental/ChaosFlesh/Source/ChaosFleshEngine/Public/ChaosFlesh/FleshComponent.h b/Engine/Plugins/Experimental/ChaosFlesh/Source/ChaosFleshEngine/Public/ChaosFlesh/FleshComponent.h
--- a/Engine/Plugins/Experimental/ChaosFlesh/Source/ChaosFleshEngine/Public/ChaosFlesh/FleshComponent.h
+++ b/Engine/Plugins/Experimental/ChaosFlesh/Source/ChaosFleshEngine/Public/ChaosFlesh/FleshComponent.h
@@ -7,6 +7,7 @@
 #include "ChaosFlesh/ChaosDeformablePhysicsComponent.h"
 #include "ChaosFlesh/FleshAsset.h"
 #include "ChaosFlesh/FleshDynamicAsset.h"
+#include "ChaosFlesh/SimulationAsset.h"
 #include "Components/MeshComponent.h"
 #include "UObject/ObjectMacros.h"
 #include "ProceduralMeshComponent.h"
@@ -69,6 +70,10 @@ public:
 	UFleshDynamicAsset* GetDynamicCollection() { return DynamicCollection; }
 	const UFleshDynamicAsset* GetDynamicCollection() const { return DynamicCollection; }
 
+	/** SimulationCollection */
+	void ResetSimulationCollection();
+	USimulationAsset* GetSimulationCollection() { return SimulationCollection; }
+	const USimulationAsset* GetSimulationCollection() const { return SimulationCollection; }
 
 	UPROPERTY()
 	TObjectPtr<UProceduralMeshComponent> Mesh;
@@ -102,6 +107,10 @@ private:
 	UPROPERTY()
 	TObjectPtr<UFleshDynamicAsset> DynamicCollection;
 
+	/** Simulator input */
+	UPROPERTY()
+	TObjectPtr<USimulationAsset> SimulationCollection;
+
 	//
 	// Sim Space
 	//
@@ -162,5 +171,6 @@ private:
 
 	int32 SimSpaceTransformIndex = INDEX_NONE;
 	int32 SimSpaceTransformGlobalIndex = INDEX_NONE;
+
 };
 
diff --git a/Engine/Plugins/Experimental/ChaosFlesh/Source/ChaosFleshEngine/Public/ChaosFlesh/SimulationAsset.h b/Engine/Plugins/Experimental/ChaosFlesh/Source/ChaosFleshEngine/Public/ChaosFlesh/SimulationAsset.h
new file mode 100644
--- /dev/null
+++ b/Engine/Plugins/Experimental/ChaosFlesh/Source/ChaosFleshEngine/Public/ChaosFlesh/SimulationAsset.h
@@ -0,0 +1,42 @@
+// Copyright Epic Games, Inc. All Rights Reserved.
+
+#pragma once
+
+#include "CoreMinimal.h"
+
+#include "GeometryCollection/ManagedArrayCollection.h"
+#include "UObject/ObjectMacros.h"
+
+#include "SimulationAsset.generated.h"
+ 
+/**
+* USimulationAsset (UObject)
+*
+* UObject wrapper for the dynamic attributes that are inputs to the tetrahedral simulation. 
+*/
+UCLASS(customconstructor)
+class CHAOSFLESHENGINE_API USimulationAsset : public UObject
+{
+	GENERATED_UCLASS_BODY()
+
+public:
+	USimulationAsset(const FObjectInitializer& ObjectInitializer = FObjectInitializer::Get());
+	void Init();
+	void Reset(const FManagedArrayCollection* InCopyFrom = nullptr);
+	void ResetAttributesFrom(const FManagedArrayCollection* InCopyFrom);
+
+	TManagedArray<int32>& GetObjectState();
+
+	const FManagedArrayCollection* GetCollection() const { return SimulationCollection.Get(); }
+	FManagedArrayCollection* GetCollection() { return SimulationCollection.Get(); }
+private:
+
+	//
+	// FManagedArrayCollection
+	// 
+	// The ManagedArrayCollection stores all the dynamic properties 
+	// for input to the simulation. This is non-const data stored on the 
+	// component. 
+	//
+	TSharedPtr<FManagedArrayCollection, ESPMode::ThreadSafe> SimulationCollection;
+};
diff --git a/Engine/Plugins/Experimental/ChaosFlesh/Source/ChaosFleshNodes/Private/Dataflow/ChaosFleshKinematicInitializationNodes.cpp b/Engine/Plugins/Experimental/ChaosFlesh/Source/ChaosFleshNodes/Private/Dataflow/ChaosFleshKinematicInitializationNodes.cpp
--- a/Engine/Plugins/Experimental/ChaosFlesh/Source/ChaosFleshNodes/Private/Dataflow/ChaosFleshKinematicInitializationNodes.cpp
+++ b/Engine/Plugins/Experimental/ChaosFlesh/Source/ChaosFleshNodes/Private/Dataflow/ChaosFleshKinematicInitializationNodes.cpp
@@ -7,6 +7,7 @@
 #include "Dataflow/DataflowEngineUtil.h"
 #include "Engine/SkeletalMesh.h"
 #include "Dataflow/DataflowInputOutput.h"
+#include "GeometryCollection/Facades/CollectionConstraintOverrideFacade.h"
 #include "GeometryCollection/Facades/CollectionKinematicBindingFacade.h"
 #include "Dataflow/DataflowNodeFactory.h"
 #include "GeometryCollection/Facades/CollectionVertexBoneWeightsFacade.h"
@@ -17,7 +18,7 @@
 
 #include UE_INLINE_GENERATED_CPP_BY_NAME(ChaosFleshKinematicInitializationNodes)
 
-//DEFINE_LOG_CATEGORY_STATIC(FKinematicInitializationNodesLog, Log, All);
+DEFINE_LOG_CATEGORY(LogKinematicInit);
 
 namespace Dataflow
 {
@@ -29,6 +30,9 @@ namespace Dataflow
 		DATAFLOW_NODE_REGISTER_CREATION_FACTORY(FKinematicOriginInsertionInitializationDataflowNode);
 		DATAFLOW_NODE_REGISTER_CREATION_FACTORY(FKinematicTetrahedralBindingsDataflowNode);
 		DATAFLOW_NODE_REGISTER_CREATION_FACTORY(FSetVerticesKinematicDataflowNode);
+		DATAFLOW_NODE_REGISTER_CREATION_FACTORY(FAuthorSceneCollisionCandidates);
+		DATAFLOW_NODE_REGISTER_CREATION_FACTORY(FBindVerticesToSkeleton);
+		DATAFLOW_NODE_REGISTER_CREATION_FACTORY(FAppendToCollectionTransformAttributeDataflowNode);
 	}
 }
 
@@ -425,6 +429,7 @@ void FSetVerticesKinematicDataflowNode::Evaluate(Dataflow::FContext& Context, co
 	}
 }
 
+
 void FKinematicBodySetupInitializationDataflowNode::Evaluate(Dataflow::FContext& Context, const FDataflowOutput* Out) const
 {
 	if (Out->IsA<DataType>(&Collection))
@@ -577,4 +582,141 @@ void FKinematicSkeletalMeshInitializationDataflowNode::Evaluate(Dataflow::FConte
 		SetValue<DataType>(Context, InCollection, &Collection);
 		SetValue<TArray<int32>>(Context, Indices, &IndicesOut);
 	}
+}
+
+
+void
+FBindVerticesToSkeleton::Evaluate(Dataflow::FContext& Context, const FDataflowOutput* Out) const
+{
+	if (Out->IsA<DataType>(&Collection))
+	{
+		DataType InCollection = GetValue<DataType>(Context, &Collection);
+
+		if (IsConnected(&VertexIndices))
+		{
+			const int32 BoneIndex = GetValue<int32>(Context, &OriginBoneIndex);
+			if (TManagedArray<FVector3f>* Vertices = InCollection.FindAttribute<FVector3f>("Vertex", FGeometryCollection::VerticesGroup))
+			{
+				const TArray<int32>& Indices = GetValue<TArray<int32>>(Context, &VertexIndices);
+				if (Indices.Num())
+				{
+					GeometryCollection::Facades::FVertexBoneWeightsFacade WeightsFacade(InCollection);
+					if (WeightsFacade.IsValid())
+					{
+						for (int32 i = 0; i < Indices.Num(); i++)
+						{
+							if (Indices[i] >= 0 && Indices[i] < Vertices->Num())
+							{
+								WeightsFacade.AddBoneWeight(Indices[i], BoneIndex, 1.0);
+							}
+						}
+					}
+				}
+			}
+		}
+		
+		SetValue<DataType>(Context, InCollection, &Collection);
+	}
+}
+
+void
+FAuthorSceneCollisionCandidates::Evaluate(Dataflow::FContext& Context, const FDataflowOutput* Out) const
+{
+	if (Out->IsA<DataType>(&Collection))
+	{
+		DataType InCollection = GetValue<DataType>(Context, &Collection);
+		const int32 BoneIndex = GetValue<int32>(Context, &OriginBoneIndex);
+
+		int32 Num = 0;
+		GeometryCollection::Facades::FConstraintOverrideCandidateFacade CnstrCandidates(InCollection);
+		CnstrCandidates.DefineSchema();
+
+		if (IsConnected(&VertexIndices))
+		{
+			if (TManagedArray<FVector3f>* Vertices = InCollection.FindAttribute<FVector3f>("Vertex", FGeometryCollection::VerticesGroup))
+			{
+				const TArray<int32>& Indices = GetValue<TArray<int32>>(Context, &VertexIndices);
+				for (int32 i = 0; i < Indices.Num(); i++)
+				{
+					if (Indices[i] >= 0 && Indices[i] < Vertices->Num())
+					{
+						GeometryCollection::Facades::FConstraintOverridesCandidateData Data;
+						Data.VertexIndex = Indices[i];
+						Data.BoneIndex = BoneIndex;
+						CnstrCandidates.Add(Data);
+						Num++;
+					}
+				}
+			}
+		}
+		else
+		{
+			TSet<int32> UniqueIndices;
+			if (bSurfaceVerticesOnly)
+			{
+				if (TManagedArray<FIntVector>* Indices =
+					InCollection.FindAttribute<FIntVector>("Indices", FGeometryCollection::FacesGroup))
+				{
+					for (int32 i = 0; i < Indices->Num(); i++)
+					{
+						for (int32 j = 0; j < 3; j++)
+						{
+							UniqueIndices.Add((*Indices)[i][j]);
+						}
+					}
+				}
+			}
+			else
+			{
+				if (TManagedArray<FIntVector4>* Indices =
+					InCollection.FindAttribute<FIntVector4>(
+						FTetrahedralCollection::TetrahedronAttribute, FTetrahedralCollection::TetrahedralGroup))
+				{
+					for (int32 i = 0; i < Indices->Num(); i++)
+					{
+						for (int32 j = 0; j < 4; j++)
+						{
+							UniqueIndices.Add((*Indices)[i][j]);
+						}
+					}
+				}
+			}
+			for (TSet<int32>::TConstIterator It = UniqueIndices.CreateConstIterator(); It; ++It)
+			{
+				GeometryCollection::Facades::FConstraintOverridesCandidateData Data;
+				Data.VertexIndex = *It;
+				Data.BoneIndex = BoneIndex;
+				CnstrCandidates.Add(Data);
+				Num++;
+			}
+		}
+		UE_LOG(LogKinematicInit, Display,
+			TEXT("'%s' - Added %d scene collision candidates."),
+			*GetName().ToString(), Num);
+		SetValue<DataType>(Context, InCollection, &Collection);
+	}
+}
+
+void
+FAppendToCollectionTransformAttributeDataflowNode::Evaluate(Dataflow::FContext& Context, const FDataflowOutput* Out) const
+{
+	if (Out->IsA<DataType>(&Collection))
+	{
+		const FTransform& Transform = GetValue<FTransform>(Context, &TransformIn);
+		FManagedArrayCollection CollectionValue = GetValue<FManagedArrayCollection>(Context, &Collection);
+		const TManagedArray<FTransform>* ConstTransformsPtr = CollectionValue.FindAttributeTyped<FTransform>(FName(AttributeName), FName(GroupName));
+		if (!ConstTransformsPtr)
+		{
+			CollectionValue.AddAttribute<FTransform>(FName(AttributeName), FName(GroupName));
+			ConstTransformsPtr = CollectionValue.FindAttributeTyped<FTransform>(FName(AttributeName), FName(GroupName));
+		}
+
+		if(ConstTransformsPtr)
+		{
+			int32 Index = CollectionValue.AddElements(1, FName(GroupName));
+			TManagedArray<FTransform>& Transforms = CollectionValue.ModifyAttribute<FTransform>(FName(AttributeName), FName(GroupName));
+			Transforms[Index] = Transform;
+		}
+		SetValue<DataType>(Context, CollectionValue, &Collection);
+	}
 }
\ No newline at end of file
diff --git a/Engine/Plugins/Experimental/ChaosFlesh/Source/ChaosFleshNodes/Public/Dataflow/ChaosFleshKinematicInitializationNodes.h b/Engine/Plugins/Experimental/ChaosFlesh/Source/ChaosFleshNodes/Public/Dataflow/ChaosFleshKinematicInitializationNodes.h
--- a/Engine/Plugins/Experimental/ChaosFlesh/Source/ChaosFleshNodes/Public/Dataflow/ChaosFleshKinematicInitializationNodes.h
+++ b/Engine/Plugins/Experimental/ChaosFlesh/Source/ChaosFleshNodes/Public/Dataflow/ChaosFleshKinematicInitializationNodes.h
@@ -1,4 +1,3 @@
-// Copyright Epic Games, Inc. All Rights Reserved.
 
 #pragma once 
 
@@ -10,6 +9,9 @@
 
 class USkeletalMesh;
 
+DECLARE_LOG_CATEGORY_EXTERN(LogKinematicInit, Verbose, All);
+
+
 USTRUCT(meta = (DataflowFlesh))
 struct FKinematicTetrahedralBindingsDataflowNode : public FDataflowNode
 {
@@ -63,9 +65,6 @@ public:
 	UPROPERTY(EditAnywhere, Category = "Dataflow")
 	float Radius = 40.f;
 
-	UPROPERTY(EditAnywhere, Category = "Dataflow")
-	FTransform Transform;
-
 	UPROPERTY(EditAnywhere, Category = "Dataflow")
 	ESkeletalSeletionMode SkeletalSelectionMode = ESkeletalSeletionMode::Dataflow_SkeletalSelection_Single;
 
@@ -218,6 +217,120 @@ public:
 	virtual void Evaluate(Dataflow::FContext& Context, const FDataflowOutput* Out) const override;
 };
 
+
+/**
+* Connects vertices to a rig so that the vertices can be animated
+*/
+USTRUCT(meta = (DataflowFlesh))
+struct FBindVerticesToSkeleton : public FDataflowNode
+{
+
+	GENERATED_USTRUCT_BODY()
+	DATAFLOW_NODE_DEFINE_INTERNAL(FBindVerticesToSkeleton, "BindVerticesToSkeleton", "Flesh", "")
+
+public:
+	typedef FManagedArrayCollection DataType;
+
+	UPROPERTY(meta = (DataflowInput, DataflowOutput, DisplayName = "Collection", DataflowPassthrough = "Collection"))
+		FManagedArrayCollection Collection;
+
+	//! Indices to use with environment collisions.  If this input is not connected, then all 
+	//! indicies are used.
+	UPROPERTY(EditAnywhere, Category = "Dataflow", meta = (DataflowInput, DisplayName = "Vertex Indices"))
+	TArray<int32> VertexIndices;
+
+	//! Bone index to use as the world raycast origin.  -1 denotes the component transform.
+	UPROPERTY(EditAnywhere, Category = "Dataflow", meta = (DataflowInput, DisplayName = "Raycast Origin Bone Index"))
+		int32 OriginBoneIndex = 0;
+
+	FBindVerticesToSkeleton(const Dataflow::FNodeParameters& InParam, FGuid InGuid = FGuid::NewGuid())
+		: FDataflowNode(InParam, InGuid)
+	{
+		RegisterInputConnection(&Collection);
+		RegisterOutputConnection(&Collection, &Collection);
+		RegisterInputConnection(&VertexIndices);
+		RegisterInputConnection(&OriginBoneIndex);
+	}
+
+	virtual void Evaluate(Dataflow::FContext& Context, const FDataflowOutput* Out) const override;
+};
+
+/**
+* Marks mesh vertices as candidates for scene collision raycasts.  Each vertex has an associated 
+* bone index to use as the origin of the raycast.  The runtime distance between the vertex and the
+* bone designates the range of the scene query depth.
+*/
+USTRUCT(meta = (DataflowFlesh))
+struct FAuthorSceneCollisionCandidates : public FDataflowNode
+{
+	GENERATED_USTRUCT_BODY()
+	DATAFLOW_NODE_DEFINE_INTERNAL(FAuthorSceneCollisionCandidates, "AuthorSceneCollisionCandidates", "Flesh", "")
+
+public:
+	typedef FManagedArrayCollection DataType;
+
+	UPROPERTY(meta = (DataflowInput, DataflowOutput, DisplayName = "Collection", DataflowPassthrough = "Collection"))
+	FManagedArrayCollection Collection;
+
+	//! Restricts vertices to only ones on the surface.  All vertices otherwise.
+	UPROPERTY(EditAnywhere, Category = "Dataflow")
+	bool bSurfaceVerticesOnly = true;
+
+	//! Indices to use with environment collisions.  If this input is not connected, then all 
+	//! indicies are used.
+	UPROPERTY(EditAnywhere, Category = "Dataflow", meta = (DataflowInput, DisplayName = "Vertex Indices"))
+	TArray<int32> VertexIndices;
+
+	//! Bone index to use as the world raycast origin.  -1 denotes the component transform.
+	UPROPERTY(EditAnywhere, Category = "Dataflow", meta = (DataflowInput, DisplayName = "Raycast Origin Bone Index"))
+	int32 OriginBoneIndex = 0;
+
+	FAuthorSceneCollisionCandidates(const Dataflow::FNodeParameters& InParam, FGuid InGuid = FGuid::NewGuid())
+		: FDataflowNode(InParam, InGuid)
+	{
+		RegisterInputConnection(&Collection);
+		RegisterOutputConnection(&Collection, &Collection);
+		RegisterInputConnection(&VertexIndices);
+		RegisterInputConnection(&OriginBoneIndex);
+	}
+
+	virtual void Evaluate(Dataflow::FContext& Context, const FDataflowOutput* Out) const override;
+};
+
+
+USTRUCT(meta = (DataflowFlesh))
+struct FAppendToCollectionTransformAttributeDataflowNode : public FDataflowNode
+{
+	GENERATED_USTRUCT_BODY()
+	DATAFLOW_NODE_DEFINE_INTERNAL(FAppendToCollectionTransformAttributeDataflowNode, "AppendToCollectionTransformAttribute", "Flesh", "")
+
+public:
+	typedef FManagedArrayCollection DataType;
+
+
+	UPROPERTY(meta = (DataflowInput, DataflowOutput, DisplayName = "Collection", DataflowPassthrough = "Collection"))
+		FManagedArrayCollection Collection;
+
+	UPROPERTY(meta = (DataflowInput, DisplayName = "Transform"))
+		FTransform TransformIn = FTransform::Identity;
+
+	UPROPERTY(EditAnywhere, Category = "Dataflow")
+		FString AttributeName = FString("ComponentTransform");
+
+	UPROPERTY(EditAnywhere, Category = "Dataflow")
+		FString GroupName = FString("ComponentTransformGroup");
+
+	FAppendToCollectionTransformAttributeDataflowNode(const Dataflow::FNodeParameters& InParam, FGuid InGuid = FGuid::NewGuid())
+		: FDataflowNode(InParam, InGuid)
+	{
+		RegisterInputConnection(&TransformIn);
+		RegisterInputConnection(&Collection);
+		RegisterOutputConnection(&Collection, &Collection);
+	}
+
+	virtual void Evaluate(Dataflow::FContext& Context, const FDataflowOutput* Out) const override;
+};
+
 namespace Dataflow
 {
 	void RegisterChaosFleshKinematicInitializationNodes();
diff --git a/Engine/Source/Runtime/Experimental/Chaos/Private/Chaos/Deformable/ChaosDeformableSolver.cpp b/Engine/Source/Runtime/Experimental/Chaos/Private/Chaos/Deformable/ChaosDeformableSolver.cpp
--- a/Engine/Source/Runtime/Experimental/Chaos/Private/Chaos/Deformable/ChaosDeformableSolver.cpp
+++ b/Engine/Source/Runtime/Experimental/Chaos/Private/Chaos/Deformable/ChaosDeformableSolver.cpp
@@ -29,6 +29,7 @@
 #include "GeometryCollection/Facades/CollectionKinematicBindingFacade.h"
 #include "GeometryCollection/Facades/CollectionVertexBoneWeightsFacade.h"
 #include "GeometryCollection/Facades/CollectionPositionTargetFacade.h"
+#include "GeometryCollection/Facades/CollectionConstraintOverrideFacade.h"
 #include "Misc/FileHelper.h"
 #include "HAL/PlatformFileManager.h"
 #include "HAL/IConsoleManager.h"
@@ -70,6 +71,7 @@ namespace Chaos::Softs
 
 	FAutoConsoleVariableRef CVarDeformableDebugParamsDrawTetrahedralParticles(TEXT("p.Chaos.DebugDraw.Deformable.TetrahedralParticle"), GDeformableDebugParams.bDoDrawTetrahedralParticles, TEXT("Debug draw the deformable solvers tetrahedron. [def: false]"));
 	FAutoConsoleVariableRef CVarDeformableDebugParamsDrawKinematicParticles(TEXT("p.Chaos.DebugDraw.Deformable.KinematicParticle"), GDeformableDebugParams.bDoDrawKinematicParticles, TEXT("Debug draw the deformables kinematic particles. [def: false]"));
+	FAutoConsoleVariableRef CVarDeformableDebugParamsDrawTransientKinematicParticles(TEXT("p.Chaos.DebugDraw.Deformable.TransientKinematicParticle"), GDeformableDebugParams.bDoDrawTransientKinematicParticles, TEXT("Debug draw the deformables transient kinematic particles. [def: false]"));
 	FAutoConsoleVariableRef CVarDeformableDebugParamsDrawRigidCollisionGeometry(TEXT("p.Chaos.DebugDraw.Deformable.RigidCollisionGeometry"), GDeformableDebugParams.bDoDrawRigidCollisionGeometry, TEXT("Debug draw the deformable solvers rigid collision geometry. [def: false]"));
 
 	FDeformableXPBDCorotatedParams GDeformableXPBDCorotatedParams;
@@ -172,6 +174,98 @@ namespace Chaos::Softs
 		}
 	}
 
+	void FDeformableSolver::UpdateTransientConstraints()
+	{
+		for(TMap<FThreadingProxy::FKey, TUniquePtr<FThreadingProxy>>::TConstIterator ProxyIt=Proxies.CreateConstIterator(); ProxyIt; ++ProxyIt)
+		{
+			const FThreadingProxy::FKey& Owner = ProxyIt.Key();
+			if (const FFleshThreadingProxy* Proxy = ProxyIt.Value()->As<FFleshThreadingProxy>())
+			{
+				if (FFleshThreadingProxy::FFleshInputBuffer* FleshInputBuffer =
+					this->CurrentInputPackage->ObjectMap.Contains(Owner) ?
+					this->CurrentInputPackage->ObjectMap[Owner]->As<FFleshThreadingProxy::FFleshInputBuffer>() :
+					nullptr)
+				{
+					GeometryCollection::Facades::FConstraintOverrideTargetFacade CnstrTargets(FleshInputBuffer->SimulationCollection);
+					if (CnstrTargets.IsValid() && CnstrTargets.Num())
+					{
+						const FIntVector2& Range = Proxy->GetSolverParticleRange();
+						const FSolverReal CurrentRatio = FSolverReal(this->Iteration) / FSolverReal(this->Property.NumSolverSubSteps);
+						const FTransform WorldToSim = Proxy->GetCurrentPointsTransform();
+
+						if (this->Iteration == 1)
+						{
+							TransientConstraintBuffer.Reserve(TransientConstraintBuffer.Num() + CnstrTargets.Num());
+							for (int32 i = 0; i < CnstrTargets.Num(); i++)
+							{
+								int32 LocalIndex = CnstrTargets.GetIndex(i);
+								int32 ParticleIndex = Range[0] + LocalIndex;
+								// Set particle kinematic state to kinematic, saving prior state.
+								TransientConstraintBuffer.Add(
+									TPair<int32, TTuple<float, float, FVector3f>>(
+										ParticleIndex,
+										TTuple<float, float, FVector3f>(
+											Evolution->Particles().InvM(ParticleIndex),
+											Evolution->Particles().PAndInvM(ParticleIndex).InvM,
+											Evolution->Particles().X(ParticleIndex))));
+
+								Evolution->Particles().InvM(ParticleIndex) = 0.f;
+								Evolution->Particles().PAndInvM(ParticleIndex).InvM = 0.f;
+							}
+						}
+
+						auto ToDouble = [](FVector3f V) { return FVector(V[0], V[1], V[2]); };
+						auto ToSingle = [](FVector V) { return FVector3f(static_cast<float>(V[0]), static_cast<float>(V[1]), static_cast<float>(V[2])); };
+
+						for (int32 i = 0; i < CnstrTargets.Num(); i++)
+						{
+							int32 LocalIndex = CnstrTargets.GetIndex(i);
+							int32 ParticleIndex = Range[0] + LocalIndex;
+
+							const FVector3f& WorldSpaceTarget = CnstrTargets.GetPosition(i);
+							FVector3f SimSpaceTarget = ToSingle(WorldToSim.TransformPosition(ToDouble(WorldSpaceTarget)));
+							const FVector3f& SimSpaceSource = TransientConstraintBuffer[ParticleIndex].Get<2>();
+
+							// Lerp from prevoius particle position to the target over the solver iterations.
+							Evolution->Particles().X(ParticleIndex) =
+								SimSpaceTarget * CurrentRatio +
+								SimSpaceSource * (static_cast<FSolverReal>(1.) - CurrentRatio);
+							Evolution->Particles().PAndInvM(ParticleIndex).P =
+								Evolution->Particles().X(ParticleIndex);
+						}
+#if WITH_EDITOR
+						if (GDeformableDebugParams.IsDebugDrawingEnabled() && GDeformableDebugParams.bDoDrawTransientKinematicParticles)
+						{
+							auto DoubleVert = [](FVector3f V) { return FVector3d(V.X, V.Y, V.Z); };
+							for (int32 i = 0; i < CnstrTargets.Num(); i++)
+							{
+								int32 LocalIndex = CnstrTargets.GetIndex(i);
+								int32 ParticleIndex = Range[0] + LocalIndex;
+								Chaos::FDebugDrawQueue::GetInstance().DrawDebugPoint(ToDouble(Evolution->Particles().X(ParticleIndex)), FColor::Orange, false, -1.0f, 0, 5);
+							}
+						}
+#endif
+					} // if has constraint overrides
+				} // if flesh input buffer
+			}
+		} // for all proxies
+	}
+
+	void FDeformableSolver::PostProcessTransientConstraints()
+	{
+		// Restore transient constraint particle kinematic state.
+		if (!TransientConstraintBuffer.IsEmpty())
+		{
+			for (TransientConstraintBufferMap::TConstIterator It = TransientConstraintBuffer.CreateConstIterator(); It; ++It)
+			{
+				const int32 ParticleIndex = It.Key();
+				Evolution->Particles().InvM(ParticleIndex) = It.Value().Get<0>();
+				Evolution->Particles().PAndInvM(ParticleIndex).InvM = It.Value().Get<1>();
+			}
+			TransientConstraintBuffer.Reset(); // retains memory
+		}
+	}
+
 	void FDeformableSolver::InitializeSimulationSpace()
 	{
 		for (int32 Index = 0; Index < this->MObjects.Num(); Index++)
@@ -242,6 +336,8 @@ namespace Chaos::Softs
 				}
 			}
 		}
+
+		UpdateTransientConstraints();
 	}
 
 	void FDeformableSolver::InitializeSimulationObject(FThreadingProxy& InProxy)
@@ -756,6 +852,11 @@ namespace Chaos::Softs
 
 			if (0 <= Index && Index < this->MObjects.Num())
 			{
+				if (TransientConstraintBuffer.Contains(Index))
+				{
+					return;
+				}
+
 				if (const UObject* Owner = this->MObjects[Index])
 				{
 					if (const FFleshThreadingProxy* Proxy = Proxies[Owner]->As<FFleshThreadingProxy>())
@@ -1013,6 +1114,7 @@ namespace Chaos::Softs
 					Iteration = i+1;
 					Update(SubDeltaTime);
 				}
+				PostProcessTransientConstraints();
 
 				Frame++;
 				EventPostSolve.Broadcast(DeltaTime);
diff --git a/Engine/Source/Runtime/Experimental/Chaos/Private/GeometryCollection/Facades/CollectionConstraintOverrideFacade.cpp b/Engine/Source/Runtime/Experimental/Chaos/Private/GeometryCollection/Facades/CollectionConstraintOverrideFacade.cpp
new file mode 100644
--- /dev/null
+++ b/Engine/Source/Runtime/Experimental/Chaos/Private/GeometryCollection/Facades/CollectionConstraintOverrideFacade.cpp
@@ -0,0 +1,210 @@
+// Copyright Epic Games, Inc. All Rights Reserved.
+
+#include "GeometryCollection/Facades/CollectionConstraintOverrideFacade.h"
+#include "GeometryCollection/GeometryCollection.h"
+
+namespace GeometryCollection::Facades
+{
+	//
+	// Attributes
+	//
+
+	const FName FConstraintOverrideCandidateFacade::GroupName("ConstraintOverrideCandidates");
+	const FName FConstraintOverrideCandidateFacade::TargetIndex("TargetIndex");
+	const FName FConstraintOverrideCandidateFacade::BoneIndex("BoneIndex");
+
+	const FName FConstraintOverrideTargetFacade::GroupName("ConstraintOverrideTargets");
+	const FName FConstraintOverrideTargetFacade::TargetIndex("TargetIndex");
+	const FName FConstraintOverrideTargetFacade::TargetPosition("TargetPosition");
+
+	//
+	// FConstraintOverrideCandidateFacade
+	//
+
+	FConstraintOverrideCandidateFacade::FConstraintOverrideCandidateFacade(FManagedArrayCollection& InCollection)
+		: ConstCollection(InCollection)
+		, Collection(&InCollection)
+		, TargetIndexAttribute(InCollection, TargetIndex, GroupName)
+		, BoneIndexAttribute(InCollection, BoneIndex, GroupName)
+	{
+		DefineSchema();
+	}
+
+	FConstraintOverrideCandidateFacade::FConstraintOverrideCandidateFacade(const FManagedArrayCollection& InCollection)
+		: ConstCollection(InCollection)
+		, Collection(nullptr)
+		, TargetIndexAttribute(InCollection, TargetIndex, GroupName)
+		, BoneIndexAttribute(InCollection, BoneIndex, GroupName)
+	{
+	}
+
+	void
+	FConstraintOverrideCandidateFacade::DefineSchema()
+	{
+		check(!IsConst());
+		TargetIndexAttribute.Add(ManageArrayAccessor::EPersistencePolicy::MakePersistent, FGeometryCollection::VerticesGroup);
+		BoneIndexAttribute.Add(ManageArrayAccessor::EPersistencePolicy::MakePersistent, FGeometryCollection::TransformGroup);
+	}
+
+	bool
+	FConstraintOverrideCandidateFacade::IsValid() const
+	{
+		return TargetIndexAttribute.IsValid() && BoneIndexAttribute.IsValid();
+	}
+
+	int32
+	FConstraintOverrideCandidateFacade::Add(FConstraintOverridesCandidateData& InputData)
+	{
+		check(!IsConst());
+		if (IsValid())
+		{
+			int32 NewIndex = TargetIndexAttribute.Get().Find(InputData.VertexIndex);
+			if (NewIndex == INDEX_NONE)
+			{
+				NewIndex = TargetIndexAttribute.AddElements(1);
+			}
+				
+			TargetIndexAttribute.Modify()[NewIndex] = InputData.VertexIndex;
+			BoneIndexAttribute.Modify()[NewIndex] = InputData.BoneIndex;
+			return NewIndex;
+		}
+		return INDEX_NONE;
+	}
+
+	void
+	FConstraintOverrideCandidateFacade::Clear()
+	{
+		check(!IsConst());
+		if (IsValid())
+		{
+			TargetIndexAttribute.GetCollection()->EmptyGroup(GroupName);
+		}
+	}
+
+	FConstraintOverridesCandidateData
+	FConstraintOverrideCandidateFacade::Get(const int32 DataIndex) const
+	{
+		FConstraintOverridesCandidateData RetVal;
+		if (IsValid() && DataIndex >= 0)
+		{
+			if (TargetIndexAttribute.Num() > DataIndex)
+			{
+				RetVal.VertexIndex = TargetIndexAttribute.Get()[DataIndex];
+			}
+			if (BoneIndexAttribute.Num() > DataIndex)
+			{
+				RetVal.BoneIndex = BoneIndexAttribute.Get()[DataIndex];
+			}
+		}
+		return RetVal;
+	}
+
+	//
+	// FConstraintOverrideTargetFacade
+	//
+
+	FConstraintOverrideTargetFacade::FConstraintOverrideTargetFacade(FManagedArrayCollection& InCollection)
+		: ConstCollection(InCollection)
+		, Collection(&InCollection)
+		, TargetIndexAttribute(InCollection, TargetIndex, GroupName)
+		, TargetPositionAttribute(InCollection, TargetPosition, GroupName)
+	{
+		DefineSchema();
+	}
+
+	FConstraintOverrideTargetFacade::FConstraintOverrideTargetFacade(const FManagedArrayCollection& InCollection)
+		: ConstCollection(InCollection)
+		, Collection(nullptr)
+		, TargetIndexAttribute(InCollection, TargetIndex, GroupName)
+		, TargetPositionAttribute(InCollection, TargetPosition, GroupName)
+	{}
+
+	void
+	FConstraintOverrideTargetFacade::DefineSchema()
+	{
+		check(!IsConst());
+		TargetIndexAttribute.Add(ManageArrayAccessor::EPersistencePolicy::MakePersistent, FGeometryCollection::VerticesGroup);
+		TargetPositionAttribute.Add(ManageArrayAccessor::EPersistencePolicy::MakePersistent);
+	}
+
+	bool
+	FConstraintOverrideTargetFacade::IsValid() const
+	{
+		return TargetIndexAttribute.IsValid() && TargetPositionAttribute.IsValid();
+	}
+
+	int32
+	FConstraintOverrideTargetFacade::Add(FConstraintOverridesTargetData& InputData)
+	{
+		check(!IsConst());
+		if (IsValid())
+		{
+			int32 NewIndex = TargetIndexAttribute.Get().Find(InputData.VertexIndex);
+			if (NewIndex == INDEX_NONE)
+			{
+				NewIndex = TargetIndexAttribute.AddElements(1);
+			}
+
+			TargetIndexAttribute.Modify()[NewIndex] = InputData.VertexIndex;
+			TargetPositionAttribute.Modify()[NewIndex] = InputData.PositionTarget;
+			return NewIndex;
+		}
+		return INDEX_NONE;
+	}
+
+	void
+	FConstraintOverrideTargetFacade::Clear()
+	{
+		check(!IsConst());
+		if (IsValid())
+		{
+			TargetIndexAttribute.GetCollection()->EmptyGroup(GroupName);
+		}
+	}
+
+	FConstraintOverridesTargetData
+	FConstraintOverrideTargetFacade::Get(const int32 DataIndex) const
+	{
+		FConstraintOverridesTargetData RetVal;
+		if (IsValid() && DataIndex >= 0)
+		{
+			if (TargetIndexAttribute.Num() > DataIndex)
+			{
+				RetVal.VertexIndex = TargetIndexAttribute.Get()[DataIndex];
+			}
+			if (TargetPositionAttribute.Num() > DataIndex)
+			{
+				RetVal.PositionTarget = TargetPositionAttribute.Get()[DataIndex];
+			}
+		}
+		return RetVal;
+	}
+
+	int32 
+	FConstraintOverrideTargetFacade::GetIndex(const int32 DataIndex) const
+	{
+		FConstraintOverridesTargetData RetVal;
+		if (IsValid() && DataIndex >= 0)
+		{
+			if (TargetIndexAttribute.Num() > DataIndex)
+			{
+				return TargetIndexAttribute.Get()[DataIndex];
+			}
+		}
+		return INDEX_NONE;
+	}
+
+	const FVector3f&
+	FConstraintOverrideTargetFacade::GetPosition(const int32 DataIndex) const
+	{
+		if (IsValid() && DataIndex >= 0)
+		{
+			if (TargetPositionAttribute.Num() > DataIndex)
+			{
+				return TargetPositionAttribute.Get()[DataIndex];
+			}
+		}
+		return FVector3f::ZeroVector;
+	}
+
+} // namespace GeometryCollection::Facades
\ No newline at end of file
diff --git a/Engine/Source/Runtime/Experimental/Chaos/Private/GeometryCollection/Facades/CollectionVertexBoneWeightsFacade.cpp b/Engine/Source/Runtime/Experimental/Chaos/Private/GeometryCollection/Facades/CollectionVertexBoneWeightsFacade.cpp
--- a/Engine/Source/Runtime/Experimental/Chaos/Private/GeometryCollection/Facades/CollectionVertexBoneWeightsFacade.cpp
+++ b/Engine/Source/Runtime/Experimental/Chaos/Private/GeometryCollection/Facades/CollectionVertexBoneWeightsFacade.cpp
@@ -21,6 +21,7 @@ namespace GeometryCollection::Facades
 		, BoneIndexAttribute(InCollection, BoneIndexAttributeName, FGeometryCollection::VerticesGroup, FTransformCollection::TransformGroup)
 		, BoneWeightAttribute(InCollection, BoneWeightAttributeName, FGeometryCollection::VerticesGroup, FTransformCollection::TransformGroup)
 		, ParentAttribute(InCollection, FTransformCollection::ParentAttribute, FTransformCollection::TransformGroup)
+		, VerticesAttribute(InCollection, "Vertex", FGeometryCollection::VerticesGroup)
 	{
 		DefineSchema();
 	}
@@ -31,6 +32,7 @@ namespace GeometryCollection::Facades
 		, BoneIndexAttribute(InCollection, BoneIndexAttributeName, FGeometryCollection::VerticesGroup, FTransformCollection::TransformGroup)
 		, BoneWeightAttribute(InCollection, BoneWeightAttributeName, FGeometryCollection::VerticesGroup, FTransformCollection::TransformGroup)
 		, ParentAttribute(InCollection, FTransformCollection::ParentAttribute, FTransformCollection::TransformGroup)
+		, VerticesAttribute(InCollection, "Vertex", FGeometryCollection::VerticesGroup)
 	{
 	}
 
@@ -49,9 +51,29 @@ namespace GeometryCollection::Facades
 
 	bool FVertexBoneWeightsFacade::IsValid() const
 	{
-		return BoneIndexAttribute.IsValid() && BoneWeightAttribute.IsValid() && ParentAttribute.IsValid();
+		return BoneIndexAttribute.IsValid() && BoneWeightAttribute.IsValid() && ParentAttribute.IsValid() && VerticesAttribute.IsValid();
 	}
 
+
+	//
+	//  Add Weights from a bone to a vertex 
+	//
+	void FVertexBoneWeightsFacade::AddBoneWeight(int32 VertexIndex, int32 BoneIndex, float BoneWeight)
+	{
+		TManagedArray< TArray<int32> >& IndicesArray = BoneIndexAttribute.Modify();
+		TManagedArray< TArray<float> >& WeightsArray = BoneWeightAttribute.Modify();
+		const TManagedArray<FVector3f>& Vertices = VerticesAttribute.Modify();
+		if (0 <= VertexIndex && VertexIndex < Vertices.Num())
+		{
+			if (0 < BoneIndex && BoneIndex < ParentAttribute.Num())
+			{
+				IndicesArray[VertexIndex].Add(BoneIndex);
+				WeightsArray[VertexIndex].Add(BoneWeight);
+			}
+		}
+	}
+
+
 	//
 	//  Add Weights from Selection 
 	//
diff --git a/Engine/Source/Runtime/Experimental/Chaos/Public/Chaos/Deformable/ChaosDeformableSolver.h b/Engine/Source/Runtime/Experimental/Chaos/Public/Chaos/Deformable/ChaosDeformableSolver.h
--- a/Engine/Source/Runtime/Experimental/Chaos/Public/Chaos/Deformable/ChaosDeformableSolver.h
+++ b/Engine/Source/Runtime/Experimental/Chaos/Public/Chaos/Deformable/ChaosDeformableSolver.h
@@ -116,6 +116,8 @@ namespace Chaos::Softs
 		void InitializeSimulationObjects();
 		void InitializeSimulationObject(FThreadingProxy&);
 		void InitializeDeformableParticles(FFleshThreadingProxy&);
+		void UpdateTransientConstraints();
+		void PostProcessTransientConstraints();
 		void InitializeKinematicParticles(FFleshThreadingProxy&);
 		void InitializeTetrahedralConstraint(FFleshThreadingProxy&);
 		void InitializeGidBasedConstraints(FFleshThreadingProxy&);
@@ -180,6 +182,10 @@ namespace Chaos::Softs
 		TUniquePtr <TArray<Chaos::TVec4<int32>>> AllElements;
 		TUniquePtr <FTriangleMesh> SurfaceTriangleMesh;
 
+		//typedef TMap<int32, TTuple<float, Chaos::Softs::FPAndInvM, FVector3f>> TransientConstraintBufferMap;
+		typedef TMap<int32, TTuple<float, float, FVector3f>> TransientConstraintBufferMap;
+		TransientConstraintBufferMap TransientConstraintBuffer;
+
 		bool bEnableSolver = true;
 		FSolverReal Time = 0.f;
 		int32 Frame = 0;
diff --git a/Engine/Source/Runtime/Experimental/Chaos/Public/Chaos/Deformable/ChaosDeformableSolverProxy.h b/Engine/Source/Runtime/Experimental/Chaos/Public/Chaos/Deformable/ChaosDeformableSolverProxy.h
--- a/Engine/Source/Runtime/Experimental/Chaos/Public/Chaos/Deformable/ChaosDeformableSolverProxy.h
+++ b/Engine/Source/Runtime/Experimental/Chaos/Public/Chaos/Deformable/ChaosDeformableSolverProxy.h
@@ -252,6 +252,7 @@ namespace Chaos::Softs
 			typedef FFleshThreadingProxy Source;
 
 			FFleshInputBuffer(
+				const FManagedArrayCollection InSimulationCollection,
 				const FTransform& InWorldToComponentXf,
 				const FTransform& InComponentToBoneXf,
 				const int32 InSimSpaceBoneIndex,
@@ -263,6 +264,7 @@ namespace Chaos::Softs
 				const float InInflationMultiplier,
 				const UObject* InOwner = nullptr)
 				: Super(InOwner, FFleshThreadingProxy::TypeName())
+				, SimulationCollection(InSimulationCollection)
 				, WorldToComponentXf(InWorldToComponentXf)
 				, ComponentToBoneXf(InComponentToBoneXf)
 				, Transforms(TArray<FTransform>())
@@ -277,6 +279,7 @@ namespace Chaos::Softs
 			{}
 
 			FFleshInputBuffer(
+				const FManagedArrayCollection InSimulationCollection,
 				const FTransform& InWorldToComponentXf,
 				const FTransform& InComponentToBoneXf, 
 				const int32 InSimSpaceBoneIndex,
@@ -290,6 +293,7 @@ namespace Chaos::Softs
 				const float InInflationMultiplier,
 				const UObject* InOwner = nullptr)
 				: Super(InOwner, FFleshThreadingProxy::TypeName())
+				, SimulationCollection(InSimulationCollection)
 				, WorldToComponentXf(InWorldToComponentXf)
 				, ComponentToBoneXf(InComponentToBoneXf)
 				, Transforms(InTransforms)
@@ -304,6 +308,8 @@ namespace Chaos::Softs
 			{}
 			virtual ~FFleshInputBuffer() {}
 
+			const FManagedArrayCollection SimulationCollection;
+
 			FTransform WorldToComponentXf;
 			FTransform ComponentToBoneXf;
 
diff --git a/Engine/Source/Runtime/Experimental/Chaos/Public/Chaos/Deformable/ChaosDeformableSolverTypes.h b/Engine/Source/Runtime/Experimental/Chaos/Public/Chaos/Deformable/ChaosDeformableSolverTypes.h
--- a/Engine/Source/Runtime/Experimental/Chaos/Public/Chaos/Deformable/ChaosDeformableSolverTypes.h
+++ b/Engine/Source/Runtime/Experimental/Chaos/Public/Chaos/Deformable/ChaosDeformableSolverTypes.h
@@ -143,6 +143,7 @@ namespace Chaos::Softs
 	{				
 		bool bDoDrawTetrahedralParticles = false;
 		bool bDoDrawKinematicParticles = false;
+		bool bDoDrawTransientKinematicParticles = false;
 		bool bDoDrawRigidCollisionGeometry = false;
 
 
diff --git a/Engine/Source/Runtime/Experimental/Chaos/Public/GeometryCollection/Facades/CollectionConstraintOverrideFacade.h b/Engine/Source/Runtime/Experimental/Chaos/Public/GeometryCollection/Facades/CollectionConstraintOverrideFacade.h
new file mode 100644
--- /dev/null
+++ b/Engine/Source/Runtime/Experimental/Chaos/Public/GeometryCollection/Facades/CollectionConstraintOverrideFacade.h
@@ -0,0 +1,103 @@
+// Copyright Epic Games, Inc. All Rights Reserved.
+
+#pragma once
+
+#include "CoreMinimal.h"
+#include "GeometryCollection/ManagedArrayAccessor.h"
+#include "GeometryCollection/ManagedArrayCollection.h"
+#include "GeometryCollection/Facades/CollectionSelectionFacade.h"
+
+namespace GeometryCollection::Facades
+{
+	/**
+	* Transient constraint candidate data.
+	*/
+	struct CHAOS_API FConstraintOverridesCandidateData
+	{
+		FConstraintOverridesCandidateData()
+			: VertexIndex(INDEX_NONE)
+			, BoneIndex(INDEX_NONE)
+		{}
+
+		int32 VertexIndex;
+		int32 BoneIndex;
+	};
+
+	/**
+	* Transient constraint target data.
+	*/
+	struct CHAOS_API FConstraintOverridesTargetData
+	{
+		FConstraintOverridesTargetData()
+			: VertexIndex(INDEX_NONE)
+		{}
+
+		int32 VertexIndex;
+		FVector3f PositionTarget;
+	};
+
+	/**
+	* Transient constraint candidates.  Typically stored in the rest collection.
+	*/
+	class CHAOS_API FConstraintOverrideCandidateFacade
+	{
+	public:
+		static const FName GroupName;
+		static const FName TargetIndex;
+		static const FName BoneIndex;
+
+		FConstraintOverrideCandidateFacade(FManagedArrayCollection& InCollection);
+		FConstraintOverrideCandidateFacade(const FManagedArrayCollection& InCollection);
+
+		void DefineSchema();
+		bool IsConst() const { return Collection == nullptr; }
+		bool IsValid() const;
+
+		int32 Add(FConstraintOverridesCandidateData& InputData);
+		void Clear();
+		FConstraintOverridesCandidateData Get(const int32 Index) const;
+		int32 Num() const { return TargetIndexAttribute.Num(); }
+
+	private:
+		const FManagedArrayCollection& ConstCollection;
+		FManagedArrayCollection* Collection = nullptr;
+
+		TManagedArrayAccessor<int32> TargetIndexAttribute;
+		TManagedArrayAccessor<int32> BoneIndexAttribute;
+	};
+
+	/**
+	* Transient constraint targets.  Typically stored in the simulation collection.
+	*/
+	class CHAOS_API FConstraintOverrideTargetFacade
+	{
+	public:
+		static const FName GroupName;
+		static const FName TargetIndex;
+		static const FName TargetPosition;
+
+		FConstraintOverrideTargetFacade(FManagedArrayCollection& InCollection);
+		FConstraintOverrideTargetFacade(const FManagedArrayCollection& InCollection);
+
+		void DefineSchema();
+		bool IsConst() const { return Collection == nullptr; }
+		bool IsValid() const;
+
+		int32 Add(FConstraintOverridesTargetData& InputData);
+		void Clear();
+		FConstraintOverridesTargetData Get(const int32 Index) const;
+		int32 GetIndex(const int32 Index) const;
+		const FVector3f& GetPosition(const int32 Index) const;
+		int32 Num() const { return TargetIndexAttribute.Num(); }
+
+	private:
+		const FManagedArrayCollection& ConstCollection;
+		FManagedArrayCollection* Collection = nullptr;
+
+		TManagedArrayAccessor<int32> TargetIndexAttribute;
+		TManagedArrayAccessor<FVector3f> TargetPositionAttribute;
+	};
+
+
+
+} // namespace GeometryCollection::Facades
\ No newline at end of file
diff --git a/Engine/Source/Runtime/Experimental/Chaos/Public/GeometryCollection/Facades/CollectionVertexBoneWeightsFacade.h b/Engine/Source/Runtime/Experimental/Chaos/Public/GeometryCollection/Facades/CollectionVertexBoneWeightsFacade.h
--- a/Engine/Source/Runtime/Experimental/Chaos/Public/GeometryCollection/Facades/CollectionVertexBoneWeightsFacade.h
+++ b/Engine/Source/Runtime/Experimental/Chaos/Public/GeometryCollection/Facades/CollectionVertexBoneWeightsFacade.h
@@ -55,6 +55,9 @@ namespace GeometryCollection::Facades
 		/** Add bone weight based on the kinematic bindings. */
 		void AddBoneWeightsFromKinematicBindings();
 
+		/** Add bone weight based on the kinematic bindings. */
+		void AddBoneWeight(int32 VertexIndex, int32 BoneIndex, float BoneWeight);
+
 		/** Return the vertex bone indices from the collection. Null if not initialized.  */
 		const TManagedArray< TArray<int32> >* FindBoneIndices()  const { return BoneIndexAttribute.Find(); }
 		const TManagedArray< TArray<int32> >& GetBoneIndices() const { return BoneIndexAttribute.Get(); }
@@ -71,6 +74,7 @@ namespace GeometryCollection::Facades
 		TManagedArrayAccessor<TArray<int32>> BoneIndexAttribute;
 		TManagedArrayAccessor<TArray<float>> BoneWeightAttribute;
 		TManagedArrayAccessor<int32> ParentAttribute;
+		TManagedArrayAccessor<FVector3f> VerticesAttribute;
 
 	};
 
