ChaosFlesh : Performance tags output.
- Fix for Non-unity CIS

diff --git a/Engine/Plugins/Experimental/ChaosFlesh/Source/ChaosFleshEngine/Private/ChaosFlesh/ChaosDeformableGameplayComponent.cpp b/Engine/Plugins/Experimental/ChaosFlesh/Source/ChaosFleshEngine/Private/ChaosFlesh/ChaosDeformableGameplayComponent.cpp
--- a/Engine/Plugins/Experimental/ChaosFlesh/Source/ChaosFleshEngine/Private/ChaosFlesh/ChaosDeformableGameplayComponent.cpp
+++ b/Engine/Plugins/Experimental/ChaosFlesh/Source/ChaosFleshEngine/Private/ChaosFlesh/ChaosDeformableGameplayComponent.cpp
@@ -26,8 +26,13 @@
 
 FChaosEngineDeformableCVarParams CVarParams2;
 FAutoConsoleVariableRef CVarDeformableDebugParamsDrawSceneRaycasts(TEXT("p.Chaos.DebugDraw.Deformable.SceneRaycasts"), CVarParams2.bDoDrawSceneRaycasts, TEXT("Debug draw the deformables scene raycasts. [def: false]"));
+FAutoConsoleVariableRef CVarDeformableDebugParamsDrawCandidateRaycasts(TEXT("p.Chaos.DebugDraw.Deformable.CandidateRaycasts"), CVarParams2.bDoDrawCandidateRaycasts, TEXT("Debug draw the deformables scene candidate raycasts. [def: false]"));
 FAutoConsoleVariableRef CVarDeformableEnvCollisionsLineTracesBatchSize(TEXT("p.Chaos.Deformable.EnvCollisionsLineTracesBatchSize"), CVarParams2.EnvCollisionsLineTraceBatchSize, TEXT("Batch size for FleshComponent env collsions. [def: 10]"));
 
+#define PERF_SCOPE(X) SCOPE_CYCLE_COUNTER(X); TRACE_CPUPROFILER_EVENT_SCOPE(X);
+DECLARE_CYCLE_STAT(TEXT("Chaos.Deformable.UDeformableGameplayComponent.PreSolverUpdate"), STAT_ChaosDeformable_UDeformableGameplayComponent_PreSolverUpdate, STATGROUP_Chaos);
+DECLARE_CYCLE_STAT(TEXT("Chaos.Deformable.UDeformableGameplayComponent.DetectEnvironmentCollisions"), STAT_ChaosDeformable_UDeformableGameplayComponent_DetectEnvironmentCollisions, STATGROUP_Chaos);
+
 //DECLARE_CYCLE_STAT(TEXT("Chaos.Deformable.UDeformableGameplayComponent.StreamStaticMeshGeometry"), STAT_ChaosDeformable_UDeformableGameplayComponent_StreamStaticMeshGeometry, STATGROUP_Chaos);
 
 UDeformableGameplayComponent::UDeformableGameplayComponent(const FObjectInitializer& ObjectInitializer)
@@ -43,6 +48,7 @@ UDeformableGameplayComponent::~UDeformableGameplayComponent()
 
 void UDeformableGameplayComponent::PreSolverUpdate()
 {
+	PERF_SCOPE(STAT_ChaosDeformable_UDeformableGameplayComponent_PreSolverUpdate)
 	if (RigBoundRayCasts.bEnableRigBoundRaycasts)
 	{
 		DetectEnvironmentCollisions(RigBoundRayCasts.MaxNumTests, RigBoundRayCasts.bTestDownOnly, RigBoundRayCasts.TestRange, RigBoundRayCasts.CollisionChannel.GetValue());
@@ -52,6 +58,8 @@ void UDeformableGameplayComponent::PreSolverUpdate()
 
 void UDeformableGameplayComponent::DetectEnvironmentCollisions(const int32 MaxNumTests, const bool bTestDownOnly, const float TestRange, const ECollisionChannel CollisionChannel)
 {
+	PERF_SCOPE(STAT_ChaosDeformable_UDeformableGameplayComponent_DetectEnvironmentCollisions)
+
 	auto ToDouble = [](FVector3f V) { return FVector(V[0], V[1], V[2]); };
 	auto ToSingle = [](FVector V) { return FVector3f(static_cast<float>(V[0]), static_cast<float>(V[1]), static_cast<float>(V[2])); };
 
@@ -227,6 +235,13 @@ void UDeformableGameplayComponent::DetectEnvironmentCollisions(const int32 MaxNu
 			WorldBoneLocations.Add(TTuple<int32, FVector>(OriginBoneIndex, WorldOrigin));
 		}
 
+#if WITH_EDITOR
+		if (CVarParams2.IsDebugDrawingEnabled() && CVarParams2.bDoDrawCandidateRaycasts)
+		{
+			Chaos::FDebugDrawQueue::GetInstance().DrawDebugDirectionalArrow(WorldOrigin, WorldSkinnedPos, 2.0f, FColor::Blue, false, -1.0f, 0, 0.75f);
+		}
+#endif
+
 		if (bTestDownOnly)
 		{
 			FVector Direction = WorldSkinnedPos - WorldOrigin;
diff --git a/Engine/Plugins/Experimental/ChaosFlesh/Source/ChaosFleshEngine/Public/ChaosFlesh/ChaosDeformableTypes.h b/Engine/Plugins/Experimental/ChaosFlesh/Source/ChaosFleshEngine/Public/ChaosFlesh/ChaosDeformableTypes.h
--- a/Engine/Plugins/Experimental/ChaosFlesh/Source/ChaosFleshEngine/Public/ChaosFlesh/ChaosDeformableTypes.h
+++ b/Engine/Plugins/Experimental/ChaosFlesh/Source/ChaosFleshEngine/Public/ChaosFlesh/ChaosDeformableTypes.h
@@ -2,6 +2,7 @@
 
 #pragma once
 #include "CoreMinimal.h"
+#include "Chaos/DebugDrawQueue.h"
 #include "UObject/ObjectMacros.h"
 #include "Engine/EngineBaseTypes.h"
 
@@ -28,6 +29,7 @@ struct FChaosEngineDeformableCVarParams
 	bool bDoDrawSkeletalMeshBindingPositions = false;
 	float DrawSkeletalMeshBindingPositionsSimulationBlendWeight = 1.f;
 	bool bDoDrawSceneRaycasts = false;
+	bool bDoDrawCandidateRaycasts = false;
 
 	int32 EnvCollisionsLineTraceBatchSize = 10;
 
