ChaosFlesh: dataflow nodes for muscle sim setup

diff --git a/Engine/Plugins/Experimental/ChaosFlesh/Source/ChaosFleshNodes/Private/Dataflow/ChaosFleshEngineAssetNodes.cpp b/Engine/Plugins/Experimental/ChaosFlesh/Source/ChaosFleshNodes/Private/Dataflow/ChaosFleshEngineAssetNodes.cpp
--- a/Engine/Plugins/Experimental/ChaosFlesh/Source/ChaosFleshNodes/Private/Dataflow/ChaosFleshEngineAssetNodes.cpp
+++ b/Engine/Plugins/Experimental/ChaosFlesh/Source/ChaosFleshNodes/Private/Dataflow/ChaosFleshEngineAssetNodes.cpp
@@ -501,36 +501,45 @@ void FGenerateOriginInsertionNode::Evaluate(Dataflow::FContext& Context, const F
 		// these via an input on the node...
 		//
 		auto DoubleVert = [](FVector3f V) { return FVector3d(V.X, V.Y, V.Z); };
-		// Origin vertices
-		if (!InOriginIndices.IsEmpty())
+		if (TManagedArray<int32>* ComponentIndex = InCollection.FindAttribute<int32>("ComponentIndex", FGeometryCollection::VerticesGroup))
 		{
-			for (int32 i = 0; i < InOriginIndices.Num(); ++i)
+			// Origin vertices
+			if (!InOriginIndices.IsEmpty())
 			{
-				if (InOriginIndices[i] < Vertex->Num())
+				for (int32 i = 0; i < InOriginIndices.Num(); ++i)
 				{
-					for (int32 j = 0; j < Vertex->Num(); ++j)
+					if (InOriginIndices[i] < Vertex->Num())
 					{
-						if (FVector::Distance(DoubleVert((*Vertex)[InOriginIndices[i]]), DoubleVert((*Vertex)[j])) < Radius)
+						for (int32 j = 0; j < Vertex->Num(); ++j)
 						{
-							OutOriginIndices.Add(j);
+							if ((*ComponentIndex)[InOriginIndices[i]] == (*ComponentIndex)[j] 
+								&& (*ComponentIndex)[InOriginIndices[i]] >= 0
+								&& (*ComponentIndex)[j] >= 0
+								&& ((*Vertex)[InOriginIndices[i]] - (*Vertex)[j]).Size() < Radius)
+							{
+								OutOriginIndices.Add(j);
+							}
 						}
 					}
 				}
 			}
-		}
 
-		// Insertion vertices
-		if (!InInsertionIndices.IsEmpty())
-		{
-			for (int32 i = 0; i < InInsertionIndices.Num(); ++i)
+			// Insertion vertices
+			if (!InInsertionIndices.IsEmpty())
 			{
-				if (InInsertionIndices[i] < Vertex->Num())
+				for (int32 i = 0; i < InInsertionIndices.Num(); ++i)
 				{
-					for (int32 j = 0; j < Vertex->Num(); ++j)
+					if (InInsertionIndices[i] < Vertex->Num())
 					{
-						if (FVector::Distance(DoubleVert((*Vertex)[InInsertionIndices[i]]), DoubleVert((*Vertex)[j])) < Radius)
+						for (int32 j = 0; j < Vertex->Num(); ++j)
 						{
-							OutInsertionIndices.Add(j);
+							if ((*ComponentIndex)[InInsertionIndices[i]] == (*ComponentIndex)[j]
+								&& (*ComponentIndex)[InInsertionIndices[i]] >= 0
+								&& (*ComponentIndex)[j] >= 0
+								&& ((*Vertex)[InInsertionIndices[i]] - (*Vertex)[j]).Size() < Radius)
+							{
+								OutInsertionIndices.Add(j);
+							}
 						}
 					}
 				}
diff --git a/Engine/Plugins/Experimental/ChaosFlesh/Source/ChaosFleshNodes/Private/Dataflow/ChaosFleshPositionTargetInitializationNodes.cpp b/Engine/Plugins/Experimental/ChaosFlesh/Source/ChaosFleshNodes/Private/Dataflow/ChaosFleshPositionTargetInitializationNodes.cpp
--- a/Engine/Plugins/Experimental/ChaosFlesh/Source/ChaosFleshNodes/Private/Dataflow/ChaosFleshPositionTargetInitializationNodes.cpp
+++ b/Engine/Plugins/Experimental/ChaosFlesh/Source/ChaosFleshNodes/Private/Dataflow/ChaosFleshPositionTargetInitializationNodes.cpp
@@ -407,12 +407,15 @@ void FSetVertexTrianglePositionTargetBindingDataflowNode::Evaluate(Dataflow::FCo
 						}
 						TArray<Chaos::TVector<float, 3>> IndicesPositions; 
 						IndicesPositions.SetNum(ActualParticleCount);
+						TArray<int32> IndicesMap;
+						IndicesMap.SetNum(ActualParticleCount);
 						int32 CurrentParticleIndex = 0;
 						for (int32 i = 0; i < GlobalIndex.Num(); i++)
 						{
 							if (GlobalIndex[i].Num() > 0)
 							{
 								IndicesPositions[CurrentParticleIndex] = (*Vertices)[(*Indices)[GlobalIndex[i][0]][LocalIndex[i][0]]];
+								IndicesMap[CurrentParticleIndex] = (*Indices)[GlobalIndex[i][0]][LocalIndex[i][0]];
 								CurrentParticleIndex += 1;
 							}
 						}
@@ -480,19 +483,19 @@ void FSetVertexTrianglePositionTargetBindingDataflowNode::Evaluate(Dataflow::FCo
 
 							int32 TriangleIndex = (*ComponentIndex)[(*Indices)[i][0]];
 							int32 MinIndex = -1;
-							float MinDis = FLT_MAX;
+							float MinDis = SphereRadius;
 							Chaos::TVector<float, 3> ClosestBary(0.f);
 							for (int32 j = 0; j < TriangleIntersections.Num(); j++)
 							{
-								if ((*ComponentIndex)[TriangleIntersections[j]] != TriangleIndex)
+								if ((*ComponentIndex)[IndicesMap[TriangleIntersections[j]]] >= 0 && TriangleIndex >= 0 && (*ComponentIndex)[IndicesMap[TriangleIntersections[j]]] != TriangleIndex)
 								{
-									Chaos::TVector<float, 3> Bary, TriPos0((*Vertices)[(*Indices)[i][0]]), TriPos1((*Vertices)[(*Indices)[i][1]]), TriPos2((*Vertices)[(*Indices)[i][2]]), ParticlePos((*Vertices)[TriangleIntersections[j]]);
+									Chaos::TVector<float, 3> Bary, TriPos0((*Vertices)[(*Indices)[i][0]]), TriPos1((*Vertices)[(*Indices)[i][1]]), TriPos2((*Vertices)[(*Indices)[i][2]]), ParticlePos((*Vertices)[IndicesMap[TriangleIntersections[j]]]);
 									Chaos::TVector<Chaos::FRealSingle, 3> ClosestPoint = Chaos::FindClosestPointAndBaryOnTriangle(TriPos0, TriPos1, TriPos2, ParticlePos, Bary);
-									Chaos::FRealSingle CurrentDistance = ((*Vertices)[TriangleIntersections[j]] - ClosestPoint).Size();
+									Chaos::FRealSingle CurrentDistance = ((*Vertices)[IndicesMap[TriangleIntersections[j]]] - ClosestPoint).Size();
 									if (CurrentDistance < MinDis)
 									{
 										MinDis = CurrentDistance;
-										MinIndex = TriangleIntersections[j];
+										MinIndex = IndicesMap[TriangleIntersections[j]];
 										ClosestBary = Bary;
 									}
 
