ChaosFlesh : Sparsely update rigid collider transforms to track current state.

diff --git a/Engine/Plugins/Experimental/ChaosFlesh/Source/ChaosFleshEngine/Private/ChaosFlesh/ChaosDeformableCollisionsComponent.cpp b/Engine/Plugins/Experimental/ChaosFlesh/Source/ChaosFleshEngine/Private/ChaosFlesh/ChaosDeformableCollisionsComponent.cpp
--- a/Engine/Plugins/Experimental/ChaosFlesh/Source/ChaosFleshEngine/Private/ChaosFlesh/ChaosDeformableCollisionsComponent.cpp
+++ b/Engine/Plugins/Experimental/ChaosFlesh/Source/ChaosFleshEngine/Private/ChaosFlesh/ChaosDeformableCollisionsComponent.cpp
@@ -31,7 +31,7 @@ void UDeformableCollisionsComponent::AddStaticMeshComponent(UStaticMeshComponent
 {
 	PERF_SCOPE(STAT_ChaosDeformable_UDeformableCollisionsComponent_AddStaticMeshComponent);
 
-		if (StaticMeshComponent)
+	if (StaticMeshComponent)
 	{
 		if (!CollisionBodies.Contains(StaticMeshComponent))
 		{
@@ -80,9 +80,9 @@ UDeformableCollisionsComponent::NewDeformableData()
 {
 	PERF_SCOPE(STAT_ChaosDeformable_UDeformableCollisionsComponent_NewDeformableData);
 
-	TArray < Chaos::Softs::FCollisionObjectAddedBodies> AddedBodiesData;
-	TArray < Chaos::Softs::FCollisionObjectRemovedBodies> RemovedBodiesData;
-	TArray < Chaos::Softs::FCollisionObjectUpdatedBodies> UpdateBodiesData;
+	TArray<Chaos::Softs::FCollisionObjectAddedBodies> AddedBodiesData;
+	TArray<Chaos::Softs::FCollisionObjectRemovedBodies> RemovedBodiesData;
+	TArray<Chaos::Softs::FCollisionObjectUpdatedBodies> UpdateBodiesData;
 
 	for (auto& CollisionBody : AddedBodies)
 	{
@@ -98,29 +98,40 @@ UDeformableCollisionsComponent::NewDeformableData()
 					FVector Scale = P.GetScale3D();
 					P.RemoveScaling();
 
-					for( int8 Index = 0; Index < BodySetup->AggGeom.SphereElems.Num(); Index++)
+					const int32 NumShapes = BodySetup->AggGeom.SphereElems.Num() + BodySetup->AggGeom.BoxElems.Num() + BodySetup->AggGeom.ConvexElems.Num();
+
+					TArray<FTransform>& PrevXfs = CollisionBodiesPrevXf.FindOrAdd(CollisionBody);
+					PrevXfs.SetNum(NumShapes);
+					int32 PrevXfIndex = 0;
+
+					// Sphere
+					for( int8 Index = 0; Index < BodySetup->AggGeom.SphereElems.Num(); Index++, PrevXfIndex++)
 					{
 						FKSphereElem& S = BodySetup->AggGeom.SphereElems[Index];
 						FTransform ShapeTransform = FTransform(Scale * S.Center) * P;
 						Geometry = new ::Chaos::TSphere<Chaos::FReal, 3>(FVector(0), S.Radius * Scale.GetMax());
+						PrevXfs[PrevXfIndex] = ShapeTransform;
 						AddedBodiesData.Add(Chaos::Softs::FCollisionObjectAddedBodies(
 							{ CollisionBody, Chaos::Softs::ERigidCollisionShapeType::Sphere, Index },
 							ShapeTransform, "", Geometry));
 					}
 
-					for (int8 Index = 0; Index < BodySetup->AggGeom.BoxElems.Num(); Index++)
+					// Box
+					for (int8 Index = 0; Index < BodySetup->AggGeom.BoxElems.Num(); Index++, PrevXfIndex++)
 					{
 						FKBoxElem& B = BodySetup->AggGeom.BoxElems[Index];
 						FTransform ShapeTransform = FTransform(Scale * B.Center) * P;
 						Geometry = new TBox<Chaos::FReal, 3>(
 							-0.5 * Scale * FVector(B.X, B.Y, B.Z),
 							0.5 * Scale * FVector(B.X, B.Y, B.Z));
+						PrevXfs[PrevXfIndex] = ShapeTransform;
 						AddedBodiesData.Add(Chaos::Softs::FCollisionObjectAddedBodies(
 							{ CollisionBody, Chaos::Softs::ERigidCollisionShapeType::Box, Index },
 							ShapeTransform, "", Geometry));
 					}
 
-					for (int8 Index = 0; Index < BodySetup->AggGeom.ConvexElems.Num(); Index++)
+					// Convex
+					for (int8 Index = 0; Index < BodySetup->AggGeom.ConvexElems.Num(); Index++, PrevXfIndex++)
 					{
 						FKConvexElem& C = BodySetup->AggGeom.ConvexElems[Index];
 						if (C.VertexData.Num())
@@ -134,6 +145,7 @@ UDeformableCollisionsComponent::NewDeformableData()
 								Vertices[i] = C.VertexData[i] * Scale;
 							}
 							Geometry = new FConvex(Vertices, 0.f);
+							PrevXfs[PrevXfIndex] = ShapeTransform;
 							AddedBodiesData.Add(Chaos::Softs::FCollisionObjectAddedBodies(
 								{ CollisionBody, Chaos::Softs::ERigidCollisionShapeType::Convex, Index},
 								ShapeTransform, "", Geometry));
@@ -144,12 +156,15 @@ UDeformableCollisionsComponent::NewDeformableData()
 		}
 	}
 
-
 	for (auto& RemovedBody : RemovedBodies) 
 	{
 		if (RemovedBody)
 		{
 			RemovedBodiesData.Add({  {RemovedBody, Chaos::Softs::ERigidCollisionShapeType::Unknown, INDEX_NONE} });
+			if (CollisionBodiesPrevXf.Contains(RemovedBody))
+			{
+				CollisionBodiesPrevXf.Remove(RemovedBody);
+			}
 		}
 	}
 	AddedBodies.Empty();
@@ -159,23 +174,60 @@ UDeformableCollisionsComponent::NewDeformableData()
 	{
 		if (CollisionBody)
 		{
-			FTransform P = CollisionBody->GetComponentToWorld();
-			UpdateBodiesData.Add({ {CollisionBody, Chaos::Softs::ERigidCollisionShapeType::Unknown, INDEX_NONE},P });
-		}
-	}
-
-
-	return FDataMapValue(new Chaos::Softs::FCollisionManagerProxy::FCollisionsInputBuffer(
-		AddedBodiesData, RemovedBodiesData, UpdateBodiesData, this));
-}
-
-
-
-
+			if (UStaticMesh* StaticMesh = CollisionBody->GetStaticMesh())
+			{
+				if (UBodySetup* BodySetup = StaticMesh->GetBodySetup())
+				{
+					FTransform P = CollisionBody->GetComponentToWorld();
+					FVector Scale = P.GetScale3D();
+					P.RemoveScaling();
 
+					const int32 NumShapes = BodySetup->AggGeom.SphereElems.Num() + BodySetup->AggGeom.BoxElems.Num() + BodySetup->AggGeom.ConvexElems.Num();
 
+					TArray<FTransform>& PrevXfs = CollisionBodiesPrevXf.FindOrAdd(CollisionBody);
+					PrevXfs.SetNum(NumShapes);
+					int32 PrevShapeXfIndex = 0;
 
+					for (int8 Index = 0; Index < BodySetup->AggGeom.SphereElems.Num(); Index++, PrevShapeXfIndex++)
+					{
+						FKSphereElem& S = BodySetup->AggGeom.SphereElems[Index];
+						FTransform ShapeTransform = FTransform(Scale * S.Center) * P;
+						if (!ShapeTransform.Equals(PrevXfs[Index]))
+						{
+							PrevXfs[Index] = ShapeTransform;
+							UpdateBodiesData.Add({ {CollisionBody, Chaos::Softs::ERigidCollisionShapeType::Sphere, Index},ShapeTransform });
+						}
+					}
 
+					for (int8 Index = 0; Index < BodySetup->AggGeom.BoxElems.Num(); Index++, PrevShapeXfIndex++)
+					{
+						FKBoxElem& B = BodySetup->AggGeom.BoxElems[Index];
+						FTransform ShapeTransform = FTransform(Scale * B.Center) * P;
+						if (!ShapeTransform.Equals(PrevXfs[Index]))
+						{
+							PrevXfs[Index] = ShapeTransform;
+							UpdateBodiesData.Add({ {CollisionBody, Chaos::Softs::ERigidCollisionShapeType::Box, Index},ShapeTransform });
+						}
+					}
 
+					for (int8 Index = 0; Index < BodySetup->AggGeom.ConvexElems.Num(); Index++, PrevShapeXfIndex++)
+					{
+						FKConvexElem& C = BodySetup->AggGeom.ConvexElems[Index];
+						if (C.VertexData.Num())
+						{
+							FTransform ShapeTransform = C.GetTransform() * P;
+							if (!ShapeTransform.Equals(PrevXfs[Index]))
+							{
+								PrevXfs[Index] = ShapeTransform;
+								UpdateBodiesData.Add({ {CollisionBody, Chaos::Softs::ERigidCollisionShapeType::Convex, Index},ShapeTransform });
+							}
+						}
+					}
+				}
+			}
+		}
+	}
 
-     
+	return FDataMapValue(new Chaos::Softs::FCollisionManagerProxy::FCollisionsInputBuffer(
+		AddedBodiesData, RemovedBodiesData, UpdateBodiesData, this));
+}
diff --git a/Engine/Plugins/Experimental/ChaosFlesh/Source/ChaosFleshEngine/Public/ChaosFlesh/ChaosDeformableCollisionsComponent.h b/Engine/Plugins/Experimental/ChaosFlesh/Source/ChaosFleshEngine/Public/ChaosFlesh/ChaosDeformableCollisionsComponent.h
--- a/Engine/Plugins/Experimental/ChaosFlesh/Source/ChaosFleshEngine/Public/ChaosFlesh/ChaosDeformableCollisionsComponent.h
+++ b/Engine/Plugins/Experimental/ChaosFlesh/Source/ChaosFleshEngine/Public/ChaosFlesh/ChaosDeformableCollisionsComponent.h
@@ -31,23 +31,19 @@ public:
 	virtual FThreadingProxy* NewProxy() override;
 	virtual FDataMapValue NewDeformableData() override;
 
-
 	UFUNCTION(BlueprintCallable, Category = "Physics")
 	void AddStaticMeshComponent(UStaticMeshComponent* StaticMeshComponent);
-	//void AddStaticMeshComponents(const TArray<TObjectPtr<UStaticMeshComponent>>& StaticMeshComponents);
 	
 	UFUNCTION(BlueprintCallable, Category = "Physics")
 	void RemoveStaticMeshComponent(UStaticMeshComponent* StaticMeshComponent);
-	//void RemoveStaticMeshComponents(const TArray<TObjectPtr<UStaticMeshComponent>>& StaticMeshComponents);
 
 	UPROPERTY(BlueprintReadOnly, Category = "Chaos Deformable")
-	TArray< TObjectPtr<UStaticMeshComponent> >  CollisionBodies;
+	TArray<TObjectPtr<UStaticMeshComponent>>  CollisionBodies;
 
 protected:
-	
-	TArray< UStaticMeshComponent* >  RemovedBodies;
-	TArray< UStaticMeshComponent* >  AddedBodies;
+	TArray<UStaticMeshComponent*> RemovedBodies;
+	TArray<UStaticMeshComponent*> AddedBodies;
 
-	
 	TMap<TObjectPtr<UObject>, FThreadingProxy*> CollisionsMap;
+	TMap<TObjectPtr<UStaticMeshComponent>, TArray<FTransform>> CollisionBodiesPrevXf;
 };
diff --git a/Engine/Source/Runtime/Experimental/Chaos/Private/Chaos/Deformable/ChaosDeformableSolver.cpp b/Engine/Source/Runtime/Experimental/Chaos/Private/Chaos/Deformable/ChaosDeformableSolver.cpp
--- a/Engine/Source/Runtime/Experimental/Chaos/Private/Chaos/Deformable/ChaosDeformableSolver.cpp
+++ b/Engine/Source/Runtime/Experimental/Chaos/Private/Chaos/Deformable/ChaosDeformableSolver.cpp
@@ -655,6 +655,17 @@ namespace Chaos::Softs
 							Proxy.CollisionBodies.Remove(KeyToRemove);
 						}
 					}
+
+					// Do updates
+					for (auto& UpdateBody : CollisionsInputBuffer->Updated)
+					{
+						if (Proxy.CollisionBodies.Contains(UpdateBody.Key))
+						{
+							FCollisionObjectParticleHandel* ParticleHandle = Proxy.CollisionBodies.Find(UpdateBody.Key);
+							Evolution->CollisionParticles().X(ParticleHandle->ParticleIndex) = UpdateBody.Transform.GetTranslation();
+							Evolution->CollisionParticles().R(ParticleHandle->ParticleIndex) = UpdateBody.Transform.GetRotation();
+						}
+					}
 				}
 			}
 		}
@@ -1375,7 +1386,7 @@ namespace Chaos::Softs
 									const FSphere& SphereGeometry = Geometry->GetObjectChecked<FSphere>();
 									FVector Center = ToFVector(CollisionParticles.X(Index)) + SphereGeometry.GetCenter();
 									FReal Radius = SphereGeometry.GetRadius();
-									Chaos::FDebugDrawQueue::GetInstance().DrawDebugSphere(Center, Radius, 12, FColor::Red, false, -1.0f, 1, 5.f);
+									Chaos::FDebugDrawQueue::GetInstance().DrawDebugSphere(Center, Radius, 12, FColor::Red, false, -1.0f, 0, 1.f);
 								}
 								else if (GeomType == ImplicitObjectType::Box)
 								{
@@ -1383,7 +1394,7 @@ namespace Chaos::Softs
 									FVector Extent = 0.5 * (BoxGeometry.Max() - BoxGeometry.Min());
 									FVector Center = ToFVector(CollisionParticles.X(Index)) + BoxGeometry.GetCenter();
 									const FQuat& Rotation = ToFQuat(CollisionParticles.R(Index));
-									Chaos::FDebugDrawQueue::GetInstance().DrawDebugBox(Center, Extent, Rotation, FColor::Red, false, -1.0f, 1, 5.f);
+									Chaos::FDebugDrawQueue::GetInstance().DrawDebugBox(Center, Extent, Rotation, FColor::Red, false, -1.0f, 0, 1.f);
 								}
 								else if (GeomType == ImplicitObjectType::Convex)
 								{
@@ -1395,7 +1406,7 @@ namespace Chaos::Softs
 										int32 Index1 = ConvexGeometry.GetEdgeVertex(EdgeIndex, 1);
 										const  TArray<FConvex::FVec3Type>& Verts = ConvexGeometry.GetVertices();
 										Chaos::FDebugDrawQueue::GetInstance().DrawDebugLine(
-											M.TransformPosition(ToFVector(Verts[Index0])), M.TransformPosition(ToFVector(Verts[Index1])), FColor::Red, false, -1.0f, 1, 5.f);
+											M.TransformPosition(ToFVector(Verts[Index0])), M.TransformPosition(ToFVector(Verts[Index1])), FColor::Red, false, -1.0f, 0, 1.f);
 									}
 								}
 							}
