ChaosFlesh : Reset proxy pointer when component disabled.
- Added solver component to referenced objects

diff --git a/Engine/Plugins/Experimental/ChaosFlesh/Source/ChaosFleshEngine/Private/ChaosFlesh/ChaosDeformablePhysicsComponent.cpp b/Engine/Plugins/Experimental/ChaosFlesh/Source/ChaosFleshEngine/Private/ChaosFlesh/ChaosDeformablePhysicsComponent.cpp
--- a/Engine/Plugins/Experimental/ChaosFlesh/Source/ChaosFleshEngine/Private/ChaosFlesh/ChaosDeformablePhysicsComponent.cpp
+++ b/Engine/Plugins/Experimental/ChaosFlesh/Source/ChaosFleshEngine/Private/ChaosFlesh/ChaosDeformablePhysicsComponent.cpp
@@ -31,6 +31,7 @@ void UDeformablePhysicsComponent::OnDestroyPhysicsState()
 			RemoveProxy(GameThreadSolver);
 		}
 	}
+	PhysicsProxy = nullptr;
 }
 
 bool UDeformablePhysicsComponent::ShouldCreatePhysicsState() const
