Add tet mesh metrics node that authors per tet signed volume and aspect ratio. Added Chaos::TTetrahedron min altitude and aspect ratio functions.

diff --git a/Engine/Plugins/Experimental/ChaosFlesh/Source/ChaosFleshNodes/Private/Dataflow/ChaosFleshTetrahedralNodes.cpp b/Engine/Plugins/Experimental/ChaosFlesh/Source/ChaosFleshNodes/Private/Dataflow/ChaosFleshTetrahedralNodes.cpp
--- a/Engine/Plugins/Experimental/ChaosFlesh/Source/ChaosFleshNodes/Private/Dataflow/ChaosFleshTetrahedralNodes.cpp
+++ b/Engine/Plugins/Experimental/ChaosFlesh/Source/ChaosFleshNodes/Private/Dataflow/ChaosFleshTetrahedralNodes.cpp
@@ -5,6 +5,7 @@
 
 #include "Chaos/Deformable/Utilities.h"
 #include "ChaosFlesh/ChaosFlesh.h"
+#include "Chaos/Tetrahedron.h"
 #include "Chaos/Utilities.h"
 #include "Chaos/UniformGrid.h"
 #include "ChaosFlesh/FleshCollection.h"
@@ -18,6 +19,7 @@
 #include "FTetWildWrapper.h"
 #include "Generate/IsosurfaceStuffing.h"
 #include "GeometryCollection/ManagedArrayCollection.h"
+#include "GeometryCollection/Facades/CollectionTetrahedralMetricsFacade.h"
 #include "MeshDescription.h"
 #include "MeshDescriptionToDynamicMesh.h"
 #include "Rendering/SkeletalMeshLODImporterData.h"
@@ -30,6 +32,7 @@ namespace Dataflow
 {
 	void ChaosFleshTetrahedralNodes()
 	{
+		DATAFLOW_NODE_REGISTER_CREATION_FACTORY(FCalculateTetMetrics);
 		DATAFLOW_NODE_REGISTER_CREATION_FACTORY(FGenerateTetrahedralCollectionDataflowNodes);
 		DATAFLOW_NODE_REGISTER_CREATION_FACTORY(FConstructTetGridNode);
 	}
@@ -71,6 +74,99 @@ namespace Dataflow
 	}
 }
 
+//=============================================================================
+// FCalculateTetMetrics
+//=============================================================================
+
+void
+FCalculateTetMetrics::Evaluate(Dataflow::FContext& Context, const FDataflowOutput* Out) const
+{
+	if (Out->IsA<DataType>(&Collection))
+	{
+		//DataType OutCollection = GetValue<DataType>(Context, &Collection); // Deep copy
+		TUniquePtr<FFleshCollection> InCollection(GetValue<DataType>(Context, &Collection).NewCopy<FFleshCollection>());
+
+		TManagedArray<FIntVector4>* TetMesh =
+			InCollection->FindAttribute<FIntVector4>(
+				FTetrahedralCollection::TetrahedronAttribute, FTetrahedralCollection::TetrahedralGroup);
+		TManagedArray<int32>* TetrahedronStart =
+			InCollection->FindAttribute<int32>(
+				FTetrahedralCollection::TetrahedronStartAttribute, FGeometryCollection::GeometryGroup);
+		TManagedArray<int32>* TetrahedronCount =
+			InCollection->FindAttribute<int32>(
+				FTetrahedralCollection::TetrahedronCountAttribute, FGeometryCollection::GeometryGroup);
+
+		TManagedArray<FVector3f>* Vertex =
+			InCollection->FindAttribute<FVector3f>(
+				"Vertex", "Vertices");
+
+		GeometryCollection::Facades::FTetrahedralMetrics TetMetrics(*InCollection);
+		TManagedArrayAccessor<float>& SignedVolume = TetMetrics.GetSignedVolume();
+		TManagedArrayAccessor<float>& AspectRatio = TetMetrics.GetAspectRatio();
+
+		for (int32 TetMeshIdx = 0; TetMeshIdx < TetrahedronStart->Num(); TetMeshIdx++)
+		{
+			const int32 TetMeshStart = (*TetrahedronStart)[TetMeshIdx];
+			const int32 TetMeshCount = (*TetrahedronCount)[TetMeshIdx];
+			
+			float MinVol = TNumericLimits<float>::Max();
+			float MaxVol = -TNumericLimits<float>::Max();
+			double AvgVol = 0.0;
+
+			float MinAR = TNumericLimits<float>::Max();
+			float MaxAR = -TNumericLimits<float>::Max();
+			double AvgAR = 0.0;
+			
+			for (int32 i = 0; i < TetMeshCount; i++)
+			{
+				const int32 Idx = TetMeshStart + i;
+				const FIntVector4& Tet = (*TetMesh)[Idx];
+				Chaos::TTetrahedron<Chaos::FReal> Tetrahedron(
+					(*Vertex)[Tet[0]],
+					(*Vertex)[Tet[1]],
+					(*Vertex)[Tet[2]],
+					(*Vertex)[Tet[3]]);
+
+				float Vol = Tetrahedron.GetSignedVolume();
+				SignedVolume.ModifyAt(Idx, Vol);
+				MinVol = MinVol < Vol ? MinVol : Vol;
+				MaxVol = MaxVol < Vol ? Vol : MaxVol;
+				AvgVol += Vol;
+
+				float AR = Tetrahedron.GetAspectRatio();
+				AspectRatio.ModifyAt(Idx, AR);
+				MinAR = MinAR < AR ? MinAR : AR;
+				MaxAR = MaxAR < AR ? AR : MaxAR;
+				AvgAR += AR;
+			}
+			if (TetMeshCount)
+			{
+				AvgVol /= TetMeshCount;
+				AvgAR /= TetMeshCount;
+			}
+			else
+			{
+				MinVol = MaxVol = 0.0f;
+				MinAR = MaxAR = 0.0f;
+			}
+
+			UE_LOG(LogChaosFlesh, Display,
+				TEXT("'%s' - Tet mesh %d of %d stats:\n"
+				"    Num Tetrahedra: %d\n"
+				"    Volume (min, avg, max): %g, %g, %g\n"
+				"    Aspect ratio (min, avg, max): %g, %g, %g"),
+				*GetName().ToString(),
+				(TetMeshIdx+1), TetrahedronStart->Num(),
+				TetMeshCount,
+				MinVol, AvgVol, MaxVol,
+				MinAR, AvgAR, MaxAR);
+        }
+
+		//SetValue<DataType>(Context, *(FManagedArrayCollection*)OutCollection.Get(), &Collection);
+		SetValue<DataType>(Context, *(FManagedArrayCollection*)InCollection.Get(), &Collection);
+	}
+}
+
 //=============================================================================
 // FConstructTetGridNode
 //=============================================================================
diff --git a/Engine/Plugins/Experimental/ChaosFlesh/Source/ChaosFleshNodes/Public/Dataflow/ChaosFleshTetrahedralNodes.h b/Engine/Plugins/Experimental/ChaosFlesh/Source/ChaosFleshNodes/Public/Dataflow/ChaosFleshTetrahedralNodes.h
--- a/Engine/Plugins/Experimental/ChaosFlesh/Source/ChaosFleshNodes/Public/Dataflow/ChaosFleshTetrahedralNodes.h
+++ b/Engine/Plugins/Experimental/ChaosFlesh/Source/ChaosFleshNodes/Public/Dataflow/ChaosFleshTetrahedralNodes.h
@@ -17,6 +17,30 @@ namespace UE {
 	}
 }
 
+// Generate quality metrics
+USTRUCT(meta = (DataflowFlesh))
+struct FCalculateTetMetrics : public FDataflowNode
+{
+	GENERATED_USTRUCT_BODY()
+	DATAFLOW_NODE_DEFINE_INTERNAL(FCalculateTetMetrics, "AuthorTetMetrics", "Flesh", "")
+public:
+	typedef FManagedArrayCollection DataType;
+
+	// Passthrough geometry collection. Bindings are stored as standalone groups in the \p Collection, keyed by the name of the input render mesh and all available LOD's.
+	UPROPERTY(meta = (DataflowInput, DataflowOutput, DataflowPassthrough = "Collection", DisplayName = "Collection"))
+	FManagedArrayCollection Collection;
+
+	FCalculateTetMetrics(const Dataflow::FNodeParameters& InParam, FGuid InGuid = FGuid::NewGuid())
+		: FDataflowNode(InParam, InGuid)
+	{
+		RegisterInputConnection(&Collection);
+		RegisterOutputConnection(&Collection, &Collection);
+	}
+
+	virtual void Evaluate(Dataflow::FContext& Context, const FDataflowOutput* Out) const override;
+};
+
+
 USTRUCT(meta = (DataflowFlesh))
 struct FConstructTetGridNode : public FDataflowNode
 {
diff --git a/Engine/Source/Runtime/Experimental/Chaos/Private/GeometryCollection/Facades/CollectionTetrahedralBindingsFacade.cpp b/Engine/Source/Runtime/Experimental/Chaos/Private/GeometryCollection/Facades/CollectionTetrahedralBindingsFacade.cpp
--- a/Engine/Source/Runtime/Experimental/Chaos/Private/GeometryCollection/Facades/CollectionTetrahedralBindingsFacade.cpp
+++ b/Engine/Source/Runtime/Experimental/Chaos/Private/GeometryCollection/Facades/CollectionTetrahedralBindingsFacade.cpp
@@ -111,10 +111,10 @@ namespace GeometryCollection::Facades
 		const int32 Idx = MeshIdAttribute.AddElements(1);
 		MeshIdAttribute.Modify()[Idx] = GroupName.ToString();
 
-		Parents.Release();
-		Weights.Release();
-		Offsets.Release();
-		Mask.Release();
+		Parents.Reset();
+		Weights.Reset();
+		Offsets.Reset();
+		Mask.Reset();
 		FManagedArrayCollection& Collection = *MeshIdAttribute.GetCollection();
 		Parents.Reset(new TManagedArrayAccessor<FIntVector4>(Collection, ParentsAttributeName, GroupName));
 		Weights.Reset(new TManagedArrayAccessor<FVector4f>(Collection, WeightsAttributeName, GroupName));
diff --git a/Engine/Source/Runtime/Experimental/Chaos/Private/GeometryCollection/Facades/CollectionTetrahedralMetricsFacade.cpp b/Engine/Source/Runtime/Experimental/Chaos/Private/GeometryCollection/Facades/CollectionTetrahedralMetricsFacade.cpp
new file mode 100644
--- /dev/null
+++ b/Engine/Source/Runtime/Experimental/Chaos/Private/GeometryCollection/Facades/CollectionTetrahedralMetricsFacade.cpp
@@ -0,0 +1,42 @@
+// Copyright Epic Games, Inc. All Rights Reserved.
+
+#include "GeometryCollection/Facades/CollectionTetrahedralMetricsFacade.h"
+#include "GeometryCollection/GeometryCollection.h"
+#include "GeometryCollection/ManagedArrayAccessor.h"
+
+namespace GeometryCollection::Facades
+{
+	// Groups
+
+	// Attributes
+	const FName FTetrahedralMetrics::SignedVolumeAttributeName = "SignedVolume";
+	const FName FTetrahedralMetrics::AspectRatioAttributeName = "AspectRatio";
+
+	FTetrahedralMetrics::FTetrahedralMetrics(FManagedArrayCollection& InCollection)
+		: SignedVolumeAttribute(InCollection, SignedVolumeAttributeName, "Tetrahedral")
+		, AspectRatioAttribute(InCollection, AspectRatioAttributeName, "Tetrahedral")
+	{
+		DefineSchema();
+	}
+
+	FTetrahedralMetrics::FTetrahedralMetrics(const FManagedArrayCollection& InCollection)
+		: SignedVolumeAttribute(InCollection, SignedVolumeAttributeName, "Tetrahedral")
+		, AspectRatioAttribute(InCollection, AspectRatioAttributeName, "Tetrahedral")
+	{}
+
+	FTetrahedralMetrics::~FTetrahedralMetrics()
+	{}
+
+	void FTetrahedralMetrics::DefineSchema()
+	{
+		check(!IsConst());
+		SignedVolumeAttribute.Add();
+		AspectRatioAttribute.Add();
+	}
+
+	bool FTetrahedralMetrics::IsValid() const
+	{
+		return SignedVolumeAttribute.IsValid() && AspectRatioAttribute.IsValid();
+	}
+
+} // namesapce GeometryCollection::Facades
diff --git a/Engine/Source/Runtime/Experimental/Chaos/Public/Chaos/Tetrahedron.h b/Engine/Source/Runtime/Experimental/Chaos/Public/Chaos/Tetrahedron.h
--- a/Engine/Source/Runtime/Experimental/Chaos/Public/Chaos/Tetrahedron.h
+++ b/Engine/Source/Runtime/Experimental/Chaos/Public/Chaos/Tetrahedron.h
@@ -97,6 +97,44 @@ namespace Chaos {
 			return TripleProduct(U, V, W) / 6;
 		}
 
+		T GetMinimumAltitude(int32* MinAltitudeVertex=nullptr) const
+		{
+			TArray<TTriangle<T>> Tris = GetTriangles();
+			TVec4<T> Distance(0, 0, 0, 0);
+			for (int32 i = 0; i < 4; i++)
+			{
+				const TTriangle<T>& Tri = Tris[i];
+				TVec3<T> Norm;
+				Distance[i] = -Tri.PhiWithNormal(X[3 - i], Norm);
+			}
+			if (MinAltitudeVertex)
+			{
+				(*MinAltitudeVertex) = 0;
+			}
+			T MinAltitude = Distance[0];
+			for (int32 i = 0; i < 4; i++)
+			{
+				if (Distance[i] < MinAltitude)
+				{
+					MinAltitude = Distance[i];
+					if (MinAltitudeVertex)
+					{
+						(*MinAltitudeVertex) = i;
+					}
+				}
+			}
+			if (MinAltitudeVertex)
+			{
+				(*MinAltitudeVertex) = 4 - (*MinAltitudeVertex);
+			}
+			return MinAltitude;
+		}
+
+		T GetAspectRatio() const
+		{
+			return GetMaxEdgeLength() / GetMinimumAltitude();
+		}
+
 		TVec3<T> GetFirstThreeBarycentricCoordinates(const TVec3<T>& Location) const 
 		{
 			TVec3<T> L1 = X[0] - X[3];
diff --git a/Engine/Source/Runtime/Experimental/Chaos/Public/GeometryCollection/Facades/CollectionTetrahedralMetricsFacade.h b/Engine/Source/Runtime/Experimental/Chaos/Public/GeometryCollection/Facades/CollectionTetrahedralMetricsFacade.h
new file mode 100644
--- /dev/null
+++ b/Engine/Source/Runtime/Experimental/Chaos/Public/GeometryCollection/Facades/CollectionTetrahedralMetricsFacade.h
@@ -0,0 +1,41 @@
+// Copyright Epic Games, Inc. All Rights Reserved.
+#pragma once
+
+#include "GeometryCollection/ManagedArrayCollection.h"
+#include "GeometryCollection/ManagedArrayAccessor.h"
+
+
+namespace GeometryCollection::Facades
+{
+	/**
+	*
+	*/
+	class CHAOS_API FTetrahedralMetrics
+	{
+	public:
+		// Groups
+
+		// Attributes
+		static const FName SignedVolumeAttributeName;
+		static const FName AspectRatioAttributeName;
+
+		FTetrahedralMetrics(FManagedArrayCollection& InCollection);
+		FTetrahedralMetrics(const FManagedArrayCollection& InCollection);
+		virtual ~FTetrahedralMetrics();
+
+		void DefineSchema();
+		bool IsConst() const { return SignedVolumeAttribute.IsConst(); }
+		bool IsValid() const;
+
+		const TManagedArrayAccessor<float>& GetSignedVolumeRO() const { return SignedVolumeAttribute; }
+		TManagedArrayAccessor<float>& GetSignedVolume() { check(!IsConst()); return SignedVolumeAttribute; }
+
+		const TManagedArrayAccessor<float>& GetAspectRatioRO() const { return AspectRatioAttribute; }
+		TManagedArrayAccessor<float>& GetAspectRatio() { check(!IsConst()); return AspectRatioAttribute; }
+
+	private:
+		TManagedArrayAccessor<float> SignedVolumeAttribute;
+		TManagedArrayAccessor<float> AspectRatioAttribute;
+	};
+
+} // namespace GeometryCollection::Facades
\ No newline at end of file
