Add the ability to not discard faces on the interior of tetrahedral meshes.

diff --git a/Engine/Plugins/Experimental/ChaosFlesh/Source/ChaosFleshNodes/Private/Dataflow/ChaosFleshRadialTetrahedronNodes.cpp b/Engine/Plugins/Experimental/ChaosFlesh/Source/ChaosFleshNodes/Private/Dataflow/ChaosFleshRadialTetrahedronNodes.cpp
--- a/Engine/Plugins/Experimental/ChaosFlesh/Source/ChaosFleshNodes/Private/Dataflow/ChaosFleshRadialTetrahedronNodes.cpp
+++ b/Engine/Plugins/Experimental/ChaosFlesh/Source/ChaosFleshNodes/Private/Dataflow/ChaosFleshRadialTetrahedronNodes.cpp
@@ -44,7 +44,7 @@ void FRadialTetrahedronDataflowNodes::Evaluate(Dataflow::FContext& Context, cons
 		TArray<FIntVector4> TetElements;
 		TArray<FVector> TetVertices;
 		RadialTetMesh(InnerRadius, OuterRadius, Height, RadialSample, AngularSample, VerticalSample, BulgeRatio * (OuterRadius - InnerRadius) / (Chaos::FReal)2., TetElements, TetVertices);
-		TArray<FIntVector3> SurfaceElements = Dataflow::GetSurfaceTriangles(TetElements);;
+		TArray<FIntVector3> SurfaceElements = Dataflow::GetSurfaceTriangles(TetElements, !bDiscardInteriorTriangles);
 		TUniquePtr<FTetrahedralCollection> RadialCollection(FTetrahedralCollection::NewTetrahedralCollection(TetVertices, SurfaceElements, TetElements));
 		InCollection->AppendGeometry(*RadialCollection.Get());
 
diff --git a/Engine/Plugins/Experimental/ChaosFlesh/Source/ChaosFleshNodes/Private/Dataflow/ChaosFleshTetrahedralNodes.cpp b/Engine/Plugins/Experimental/ChaosFlesh/Source/ChaosFleshNodes/Private/Dataflow/ChaosFleshTetrahedralNodes.cpp
--- a/Engine/Plugins/Experimental/ChaosFlesh/Source/ChaosFleshNodes/Private/Dataflow/ChaosFleshTetrahedralNodes.cpp
+++ b/Engine/Plugins/Experimental/ChaosFlesh/Source/ChaosFleshNodes/Private/Dataflow/ChaosFleshTetrahedralNodes.cpp
@@ -38,7 +38,7 @@ namespace Dataflow
 	}
 
 	// Helper to get the boundary of a tet mesh, useful for debugging / verifying output
-	TArray<FIntVector3> GetSurfaceTriangles(const TArray<FIntVector4>& Tets)
+	TArray<FIntVector3> GetSurfaceTriangles(const TArray<FIntVector4>& Tets, const bool bKeepInterior)
 	{
 		// Rotate the vector so the first element is the smallest
 		auto RotVec = [](const FIntVector3& F) -> FIntVector3
@@ -59,10 +59,15 @@ namespace Dataflow
 			Chaos::Utilities::GetTetFaces(Tets[TetIdx], TetF[0], TetF[1], TetF[2], TetF[3], false);
 			for (int32 SubIdx = 0; SubIdx < 4; ++SubIdx)
 			{
+				// A face can be shared between a maximum of 2 tets, so no need worrying about 
+				// re-adding removed faces.
 				FIntVector3 Key = RotVec(TetF[SubIdx]);
 				if (FacesSet.Contains(Key))
 				{
-					FacesSet.Remove(Key);
+					if (!bKeepInterior)
+					{
+						FacesSet.Remove(Key);
+					}
 				}
 				else
 				{
@@ -189,7 +194,7 @@ void FConstructTetGridNode::Evaluate(Dataflow::FContext& Context, const FDataflo
 
 		UE_LOG(LogChaosFlesh, Display, TEXT("TetGrid generated %d points and %d tetrahedra."), X.Num(), Tets.Num());
 
-		TArray<FIntVector3> Tris = Dataflow::GetSurfaceTriangles(Tets);
+		TArray<FIntVector3> Tris = Dataflow::GetSurfaceTriangles(Tets, !bDiscardInteriorTriangles);
 		TUniquePtr<FTetrahedralCollection> TetCollection(
 			FTetrahedralCollection::NewTetrahedralCollection(X, Tris, Tets));
 		InCollection->AppendGeometry(*TetCollection.Get());
@@ -300,7 +305,7 @@ void FGenerateTetrahedralCollectionDataflowNodes::EvaluateIsoStuffing(
 		{
 			TArray<FVector> Vertices; Vertices.SetNumUninitialized(IsosurfaceStuffing.Vertices.Num());
 			TArray<FIntVector4> Elements; Elements.SetNumUninitialized(IsosurfaceStuffing.Tets.Num());
-			TArray<FIntVector3> SurfaceElements = Dataflow::GetSurfaceTriangles(IsosurfaceStuffing.Tets);
+			TArray<FIntVector3> SurfaceElements = Dataflow::GetSurfaceTriangles(IsosurfaceStuffing.Tets, !bDiscardInteriorTriangles);
 
 			for (int32 Tdx = 0; Tdx < IsosurfaceStuffing.Tets.Num(); ++Tdx)
 			{
@@ -366,7 +371,7 @@ void FGenerateTetrahedralCollectionDataflowNodes::EvaluateTetWild(
 		UE_LOG(LogChaosFlesh, Display,TEXT("Generating tet mesh via TetWild..."));
 		if (UE::Geometry::FTetWild::ComputeTetMesh(Params, Verts, Tris, TetVerts, Tets, &Progress))
 		{
-			TArray<FIntVector3> SurfaceElements = Dataflow::GetSurfaceTriangles(Tets);
+			TArray<FIntVector3> SurfaceElements = Dataflow::GetSurfaceTriangles(Tets, !bDiscardInteriorTriangles);
 			TUniquePtr<FTetrahedralCollection> TetCollection(FTetrahedralCollection::NewTetrahedralCollection(TetVerts, SurfaceElements, Tets));
 			InCollection->AppendGeometry(*TetCollection.Get());
 
diff --git a/Engine/Plugins/Experimental/ChaosFlesh/Source/ChaosFleshNodes/Public/Dataflow/ChaosFleshRadialTetrahedronNodes.h b/Engine/Plugins/Experimental/ChaosFlesh/Source/ChaosFleshNodes/Public/Dataflow/ChaosFleshRadialTetrahedronNodes.h
--- a/Engine/Plugins/Experimental/ChaosFlesh/Source/ChaosFleshNodes/Public/Dataflow/ChaosFleshRadialTetrahedronNodes.h
+++ b/Engine/Plugins/Experimental/ChaosFlesh/Source/ChaosFleshNodes/Public/Dataflow/ChaosFleshRadialTetrahedronNodes.h
@@ -45,6 +45,9 @@ public:
 	UPROPERTY(EditAnywhere, Category = "Dataflow", meta = (ClampMin = "0.0", ClampMax = "1.0"))
 	double BulgeRatio = double(0.0);
 
+	UPROPERTY(EditAnywhere, Category = "Dataflow")
+	bool bDiscardInteriorTriangles = true;
+
 	FRadialTetrahedronDataflowNodes(const Dataflow::FNodeParameters& InParam, FGuid InGuid = FGuid::NewGuid())
 		: FDataflowNode(InParam, InGuid)
 	{
diff --git a/Engine/Plugins/Experimental/ChaosFlesh/Source/ChaosFleshNodes/Public/Dataflow/ChaosFleshTetrahedralNodes.h b/Engine/Plugins/Experimental/ChaosFlesh/Source/ChaosFleshNodes/Public/Dataflow/ChaosFleshTetrahedralNodes.h
--- a/Engine/Plugins/Experimental/ChaosFlesh/Source/ChaosFleshNodes/Public/Dataflow/ChaosFleshTetrahedralNodes.h
+++ b/Engine/Plugins/Experimental/ChaosFlesh/Source/ChaosFleshNodes/Public/Dataflow/ChaosFleshTetrahedralNodes.h
@@ -60,6 +60,9 @@ public:
 	UPROPERTY(EditAnywhere, Category = "Dataflow", meta = (ClampMin = "0.0"))
 	FVector GridDomain = FVector(10.0, 10.0, 10.0);
 
+	UPROPERTY(EditAnywhere, Category = "Dataflow")
+	bool bDiscardInteriorTriangles = true;
+
 	FConstructTetGridNode(const Dataflow::FNodeParameters& InParam, FGuid InGuid = FGuid::NewGuid())
 		: FDataflowNode(InParam, InGuid)
 	{
@@ -150,6 +153,9 @@ public:
 	UPROPERTY(meta = (DataflowInput, DataflowOutput, DisplayName = "Collection"))
 	FManagedArrayCollection Collection;
 
+	UPROPERTY(EditAnywhere, Category = "Dataflow")
+	bool bDiscardInteriorTriangles = true;
+
 	FGenerateTetrahedralCollectionDataflowNodes(const Dataflow::FNodeParameters& InParam, FGuid InGuid = FGuid::NewGuid())
 		: FDataflowNode(InParam, InGuid)
 	{
@@ -168,7 +174,7 @@ protected:
 
 namespace Dataflow
 {
-	TArray<FIntVector3> GetSurfaceTriangles(const TArray<FIntVector4>& Tets);
+	TArray<FIntVector3> GetSurfaceTriangles(const TArray<FIntVector4>& Tets, const bool bKeepInterior);
 	void ChaosFleshTetrahedralNodes();
 
 
