Add command for creating selections of tetrahedra.  Current supported critera includes by volume less than threshold, by aspect ratio greater than threshold, and by X/Y/Z coordinate less than or greater than some specified number.  Command also supports authoring a TSet of tetrahedra on the Flesh component to skip during rendering.

diff --git a/Engine/Plugins/Experimental/ChaosFlesh/Source/ChaosFleshEditor/Private/ChaosFlesh/ChaosFleshEditorPlugin.cpp b/Engine/Plugins/Experimental/ChaosFlesh/Source/ChaosFleshEditor/Private/ChaosFlesh/ChaosFleshEditorPlugin.cpp
--- a/Engine/Plugins/Experimental/ChaosFlesh/Source/ChaosFleshEditor/Private/ChaosFlesh/ChaosFleshEditorPlugin.cpp
+++ b/Engine/Plugins/Experimental/ChaosFlesh/Source/ChaosFleshEditor/Private/ChaosFlesh/ChaosFleshEditorPlugin.cpp
@@ -29,11 +29,23 @@ void IChaosFleshEditorPlugin::StartupModule()
 	if (GIsEditor && !IsRunningCommandlet())
 	{
 		EditorCommands.Add(IConsoleManager::Get().RegisterConsoleCommand(
-			TEXT("p.Chaos.Flesh.ImportFile"),
+			TEXT("ChaosDeformable.ImportFile"),
 			TEXT("Creates a FleshAsset from the input file"),
 			FConsoleCommandWithWorldAndArgsDelegate::CreateStatic(&FChaosFleshCommands::ImportFile),
 			ECVF_Default
 		));
+
+		EditorCommands.Add(IConsoleManager::Get().RegisterConsoleCommand(
+			TEXT("ChaosDeformable.FindQualifyingTetrahedra"),
+			TEXT("From the selected actor's flesh components, prints indices of tetrahedra matching our search criteria. "
+				"Use arg 'MinVol <value>' to specify a minimum tet volume; "
+				"use arg 'MaxAR <value>' to specify a maximum aspect ratio; "
+				"use 'XCoordGT <value>', 'YCoordGT <value>', 'ZCoordGT <value>' to select tets with all vertices greater than the specified value; "
+				"use 'XCoordLT <value>', 'YCoordLT <value>', 'ZCoordLT <value>' to select tets with all vertices less than the specified value; "
+				"use 'HideTets' to add indices to the flesh component's list of tets to skip drawing."),
+			FConsoleCommandWithWorldAndArgsDelegate::CreateStatic(&FChaosFleshCommands::FindQualifyingTetrahedra),
+			ECVF_Default
+		));
 	}
 
 	UThumbnailManager::Get().RegisterCustomRenderer(UFleshAsset::StaticClass(), UFleshAssetThumbnailRenderer::StaticClass());
diff --git a/Engine/Plugins/Experimental/ChaosFlesh/Source/ChaosFleshEditor/Private/ChaosFlesh/Cmd/ChaosFleshCommands.cpp b/Engine/Plugins/Experimental/ChaosFlesh/Source/ChaosFleshEditor/Private/ChaosFlesh/Cmd/ChaosFleshCommands.cpp
--- a/Engine/Plugins/Experimental/ChaosFlesh/Source/ChaosFleshEditor/Private/ChaosFlesh/Cmd/ChaosFleshCommands.cpp
+++ b/Engine/Plugins/Experimental/ChaosFlesh/Source/ChaosFleshEditor/Private/ChaosFlesh/Cmd/ChaosFleshCommands.cpp
@@ -7,17 +7,22 @@
 #include "ChaosFlesh/FleshAsset.h"
 #include "ChaosFlesh/FleshCollection.h"
 #include "ChaosFlesh/FleshCollectionUtility.h"
+#include "ChaosFlesh/FleshComponent.h"
+#include "Chaos/Tetrahedron.h"
 
 #include "AssetRegistry/AssetRegistryModule.h"
+#include "Editor.h"
+#include "Editor/EditorEngine.h"
+#include "Engine/Selection.h"
+#include "Engine/World.h"
+#include "EngineUtils.h"
 #include "GeometryCollection/TransformCollection.h"
 #include "Misc/Paths.h"
 #include "UObject/Package.h"
 
-
 DEFINE_LOG_CATEGORY_STATIC(UChaosFleshCommandsLogging, NoLogging, All);
 
 
-
 void FChaosFleshCommands::ImportFile(const TArray<FString>& Args, UWorld* World)
 {
 	if (Args.Num() == 1)
@@ -50,3 +55,252 @@ void FChaosFleshCommands::ImportFile(const TArray<FString>& Args, UWorld* World)
 		UE_LOG(UChaosFleshCommandsLogging, Error, TEXT("Failed to import file for flesh asset."));
 	}
 }
+
+void
+FChaosFleshCommands::FindQualifyingTetrahedra(const TArray<FString>& Args, UWorld* World)
+{
+	if (USelection* SelectedActors = GEditor->GetSelectedActors())
+	{
+		for (FSelectionIterator ActorIt(*SelectedActors); ActorIt; ++ActorIt)
+		{
+			if (AActor* Actor = Cast<AActor>(*ActorIt))
+			{
+				const TSet<UActorComponent*>& Components = Actor->GetComponents();
+				for (TSet<UActorComponent*>::TConstIterator CompIt = Components.CreateConstIterator(); CompIt; ++CompIt)
+				{
+					if (UFleshComponent* FleshComponent = Cast<UFleshComponent>(*CompIt))
+					{
+						if (const UFleshAsset* RestCollection = FleshComponent->GetRestCollection())
+						{
+							const FFleshCollection* FleshCollection = RestCollection->GetCollection();
+
+							const TManagedArray<FIntVector4>* TetMesh =
+								FleshCollection->FindAttribute<FIntVector4>(
+									FTetrahedralCollection::TetrahedronAttribute, FTetrahedralCollection::TetrahedralGroup);
+							const TManagedArray<int32>* TetrahedronStart =
+								FleshCollection->FindAttribute<int32>(
+									FTetrahedralCollection::TetrahedronStartAttribute, FGeometryCollection::GeometryGroup);
+							const TManagedArray<int32>* TetrahedronCount =
+								FleshCollection->FindAttribute<int32>(
+									FTetrahedralCollection::TetrahedronCountAttribute, FGeometryCollection::GeometryGroup);
+
+							const UFleshDynamicAsset* DynamicCollection = FleshComponent->GetDynamicCollection();
+							const TManagedArray<FVector3f>* Vertex =
+								DynamicCollection && DynamicCollection->FindPositions() && DynamicCollection->FindPositions()->Num() ?
+								DynamicCollection->FindPositions() :
+								FleshCollection->FindAttribute<FVector3f>("Vertex", "Vertices");
+
+							float MaxAR = TNumericLimits<float>::Max();
+							float MinVol = -TNumericLimits<float>::Max();
+							float XCoordGT = TNumericLimits<float>::Max();
+							float YCoordGT = TNumericLimits<float>::Max();
+							float ZCoordGT = TNumericLimits<float>::Max();
+							float XCoordLT = -TNumericLimits<float>::Max();
+							float YCoordLT = -TNumericLimits<float>::Max();
+							float ZCoordLT = -TNumericLimits<float>::Max();
+							bool bHideTets = false;
+							for (int32 Idx = 0; Idx < Args.Num(); Idx++)
+							{
+								if (Args[Idx].Equals(FString(TEXT("MaxAR"))))
+								{
+									if (Args.IsValidIndex(Idx + 1))
+									{
+										Idx++;
+										MaxAR = FCString::Atof(*Args[Idx]);
+									}
+								}
+								else if (Args[Idx].Equals(FString(TEXT("MinVol"))))
+								{
+									if (Args.IsValidIndex(Idx + 1))
+									{
+										Idx++;
+										MinVol = FCString::Atof(*Args[Idx]);
+									}
+								}
+
+								// XYZ greater than
+								else if (Args[Idx].Equals(FString(TEXT("XCoordGT"))))
+								{
+									if (Args.IsValidIndex(Idx + 1))
+									{
+										Idx++;
+										XCoordGT = FCString::Atof(*Args[Idx]);
+									}
+								}
+								else if (Args[Idx].Equals(FString(TEXT("YCoordGT"))))
+								{
+									if (Args.IsValidIndex(Idx + 1))
+									{
+										Idx++;
+										YCoordGT = FCString::Atof(*Args[Idx]);
+									}
+								}
+								else if (Args[Idx].Equals(FString(TEXT("ZCoordGT"))))
+								{
+									if (Args.IsValidIndex(Idx + 1))
+									{
+										Idx++;
+										ZCoordGT = FCString::Atof(*Args[Idx]);
+									}
+								}
+
+								// XYZ less than
+								else if (Args[Idx].Equals(FString(TEXT("XCoordLT"))))
+								{
+									if (Args.IsValidIndex(Idx + 1))
+									{
+										Idx++;
+										XCoordLT = FCString::Atof(*Args[Idx]);
+									}
+								}
+								else if (Args[Idx].Equals(FString(TEXT("YCoordLT"))))
+								{
+									if (Args.IsValidIndex(Idx + 1))
+									{
+										Idx++;
+										YCoordLT = FCString::Atof(*Args[Idx]);
+									}
+								}
+								else if (Args[Idx].Equals(FString(TEXT("ZCoordLT"))))
+								{
+									if (Args.IsValidIndex(Idx + 1))
+									{
+										Idx++;
+										ZCoordLT = FCString::Atof(*Args[Idx]);
+									}
+								}
+
+								else if (Args[Idx].Equals(FString(TEXT("HideTets"))))
+								{
+									bHideTets = true;
+								}
+							}
+
+							TArray<int32> Indices;
+							for (int32 TetMeshIdx = 0; TetMeshIdx < TetrahedronStart->Num(); TetMeshIdx++)
+							{
+								const int32 TetMeshStart = (*TetrahedronStart)[TetMeshIdx];
+								const int32 TetMeshCount = (*TetrahedronCount)[TetMeshIdx];
+
+								for (int32 i = 0; i < TetMeshCount; i++)
+								{
+									const int32 Idx = TetMeshStart + i;
+									const FIntVector4& Tet = (*TetMesh)[Idx];
+
+									const int32 MaxIdx = FGenericPlatformMath::Max(
+										FGenericPlatformMath::Max(Tet[0], Tet[1]),
+										FGenericPlatformMath::Max(Tet[2], Tet[3]));
+									if (MaxIdx >= Vertex->Num())
+									{
+										continue;
+									}
+
+									Chaos::TTetrahedron<Chaos::FReal> Tetrahedron(
+										(*Vertex)[Tet[0]],
+										(*Vertex)[Tet[1]],
+										(*Vertex)[Tet[2]],
+										(*Vertex)[Tet[3]]);
+
+									if (MinVol != -TNumericLimits<float>::Max())
+									{
+										float Vol = Tetrahedron.GetSignedVolume();
+										if (Vol < MinVol)
+										{
+											Indices.Add(Idx);
+											continue;
+										}
+									}
+									if (MaxAR != TNumericLimits<float>::Max())
+									{
+										float AR = Tetrahedron.GetAspectRatio();
+										if (AR > MaxAR)
+										{
+											Indices.Add(Idx);
+											continue;
+										}
+									}
+									if (XCoordGT != TNumericLimits<float>::Max() ||
+										YCoordGT != TNumericLimits<float>::Max() ||
+										ZCoordGT != TNumericLimits<float>::Max())
+									{
+										bool bAdd = true;
+										if (XCoordGT != TNumericLimits<float>::Max())
+										{
+											for (int32 j = 0; j < 4; j++) bAdd &= (*Vertex)[Tet[j]][0] > XCoordGT;
+										}
+										if (YCoordGT != TNumericLimits<float>::Max())
+										{
+											for (int32 j = 0; j < 4; j++) bAdd &= (*Vertex)[Tet[j]][1] > YCoordGT;
+										}
+										if (ZCoordGT != TNumericLimits<float>::Max())
+										{
+											for (int32 j = 0; j < 4; j++) bAdd &= (*Vertex)[Tet[j]][2] > ZCoordGT;
+										}
+										if (bAdd)
+										{
+											Indices.Add(Idx);
+											continue;
+										}
+									}
+									if (XCoordLT != -TNumericLimits<float>::Max() ||
+										YCoordLT != -TNumericLimits<float>::Max() ||
+										ZCoordLT != -TNumericLimits<float>::Max())
+									{
+										bool bAdd = true;
+										if (XCoordLT != -TNumericLimits<float>::Max())
+										{
+											for (int32 j = 0; j < 4; j++) bAdd &= (*Vertex)[Tet[j]][0] < XCoordLT;
+										}
+										if (YCoordLT != -TNumericLimits<float>::Max())
+										{
+											for (int32 j = 0; j < 4; j++) bAdd &= (*Vertex)[Tet[j]][1] < YCoordLT;
+										}
+										if (ZCoordLT != -TNumericLimits<float>::Max())
+										{
+											for (int32 j = 0; j < 4; j++) bAdd &= (*Vertex)[Tet[j]][2] < ZCoordLT;
+										}
+										if (bAdd)
+										{
+											Indices.Add(Idx);
+											continue;
+										}
+									}
+								}
+							}
+
+							if (bHideTets)
+							{
+								FleshComponent->HideTetrahedra.Append(Indices);
+							}
+
+							if (Indices.Num())
+							{
+								FString IndicesStr(TEXT("["));
+								int32 i = 0;
+								for (; i < Indices.Num() - 1; i++)
+								{
+									IndicesStr.Append(FString::Printf(TEXT("%d "), Indices[i]));
+								}
+								if (i < Indices.Num())
+								{
+									IndicesStr.Append(FString::Printf(TEXT("%d]"), Indices[i]));
+								}
+								else
+								{
+									IndicesStr.Append(TEXT("]"));
+								}
+								UE_LOG(UChaosFleshCommandsLogging, Log,
+									TEXT("ChaosDeformableCommands.FindQualifyingTetrahedra - '%s.%s' Found %d qualifying tetrahedra: \n%s"),
+									*Actor->GetName(),
+									*FleshComponent->GetName(),
+									Indices.Num(),
+									*IndicesStr);
+							}
+
+						} // end if RestCollection
+					}
+				}
+			}
+		}
+	}
+}
\ No newline at end of file
diff --git a/Engine/Plugins/Experimental/ChaosFlesh/Source/ChaosFleshEditor/Public/ChaosFlesh/Cmd/ChaosFleshCommands.h b/Engine/Plugins/Experimental/ChaosFlesh/Source/ChaosFleshEditor/Public/ChaosFlesh/Cmd/ChaosFleshCommands.h
--- a/Engine/Plugins/Experimental/ChaosFlesh/Source/ChaosFleshEditor/Public/ChaosFlesh/Cmd/ChaosFleshCommands.h
+++ b/Engine/Plugins/Experimental/ChaosFlesh/Source/ChaosFleshEditor/Public/ChaosFlesh/Cmd/ChaosFleshCommands.h
@@ -16,4 +16,16 @@ public:
 	*/
 	static void ImportFile(const TArray<FString>& Args, UWorld* World);
 
+	/**
+	* Commmand invoked from "FChaosDeformableCommands.FindHighAspectRatioTetrahedra", uses
+	* the selected FleshComponent and outputs tetrahedra indices to the log.
+	* @param Args - Supported arguments:
+	*	'MaxAR <float>' selects tetrahedra with aspect ratio greater than this value.
+	*	'MinVol <float>' selects tetrahedra with (signed) volume less than this value.
+	*	'XCoordGT <float>', 'YCoordGT <float>', 'ZCoordGT <float>' selects tetrahedra with all vertices greater than these values.
+	*	'XCoordLT <float>', 'YCoordLT <float>', 'ZCoordLT <float>' selects tetrahedra with all vertices less than these values.
+	*   'HideTets' adds selected tets to the selected flesh component's list of tets to skip drawing.
+	* @param World
+	*/
+	static void FindQualifyingTetrahedra(const TArray<FString>& Args, UWorld* World);
 };
\ No newline at end of file
diff --git a/Engine/Plugins/Experimental/ChaosFlesh/Source/ChaosFleshEngine/Private/ChaosFlesh/FleshComponent.cpp b/Engine/Plugins/Experimental/ChaosFlesh/Source/ChaosFleshEngine/Private/ChaosFlesh/FleshComponent.cpp
--- a/Engine/Plugins/Experimental/ChaosFlesh/Source/ChaosFleshEngine/Private/ChaosFlesh/FleshComponent.cpp
+++ b/Engine/Plugins/Experimental/ChaosFlesh/Source/ChaosFleshEngine/Private/ChaosFlesh/FleshComponent.cpp
@@ -445,24 +445,72 @@ void UFleshComponent::RenderProceduralMesh()
 
 					if (const FFleshCollection* Flesh = FleshAsset->GetCollection())
 					{
+						if (HideTetrahedra.Num() != NumSkippedTets)
+						{
+							const TManagedArray<FIntVector4>* Tets =
+								Flesh->FindAttribute<FIntVector4>(
+									FTetrahedralCollection::TetrahedronAttribute, FTetrahedralCollection::TetrahedralGroup);
+
+							FaceVerticesToSkip.Reset();
+							for (TSet<int32>::TConstIterator It = HideTetrahedra.CreateConstIterator(); It; ++It)
+							{
+								const int32 Idx = *It;
+								if (Idx >= 0 && Idx < Tets->Num())
+								{
+									const FIntVector4& Tet = (*Tets)[Idx];
+									for (int32 i = 0; i < 4; i++)
+									{
+										FaceVerticesToSkip.Add(Tet[i]);
+									}
+								}
+							}
+							NumSkippedTets = HideTetrahedra.Num();
+							NumSkippedFaces = 0;
+						}
+
 						int32 NumVertices = Flesh->NumElements(FGeometryCollection::VerticesGroup);
 						int32 NumFaces = Flesh->NumElements(FGeometryCollection::FacesGroup);
 						if (NumFaces && NumVertices)
 						{
-							if (RenderMesh && RenderMesh->Vertices.Num() != NumFaces * 3)
+							if (RenderMesh && RenderMesh->Vertices.Num() != (NumFaces - NumSkippedFaces) * 3)
 							{
 								ResetProceduralMesh();
 							}
 
+							// Find vertices array.  Use dynamic collection if available and size matches.
+							const TManagedArray<FVector3f>* RenderVertex = &Flesh->Vertex;
+							if (GetDynamicCollection())
+							{
+								const TManagedArray<FVector3f>& DynamicVertex = GetDynamicCollection()->GetPositions();
+								if (DynamicVertex.Num() == RenderVertex->Num())
+								{
+									RenderVertex = &DynamicVertex;
+								}
+							}
+
+							const TManagedArray<bool>* FacesVisible =
+								Flesh->FindAttribute<bool>(
+									"Visible", FGeometryCollection::FacesGroup);
+
 							if (!RenderMesh)
 							{
 								RenderMesh = new FFleshRenderMesh;
+								NumSkippedFaces = 0;
 
-								for (int i = 0; i < NumFaces; ++i)
+								for (int i = 0, j = 0; i < NumFaces; ++i)
 								{
-									const auto& P1 = Flesh->Vertex[Flesh->Indices[i][0]];
-									const auto& P2 = Flesh->Vertex[Flesh->Indices[i][1]];
-									const auto& P3 = Flesh->Vertex[Flesh->Indices[i][2]];
+									if (((FacesVisible && (*FacesVisible)[i] == false)) ||
+										(FaceVerticesToSkip.Contains(Flesh->Indices[i][0]) &&
+										 FaceVerticesToSkip.Contains(Flesh->Indices[i][1]) &&
+										 FaceVerticesToSkip.Contains(Flesh->Indices[i][2])))
+									{
+										NumSkippedFaces++;
+										continue;
+									}
+
+									const auto& P1 = (*RenderVertex)[Flesh->Indices[i][0]];
+									const auto& P2 = (*RenderVertex)[Flesh->Indices[i][1]];
+									const auto& P3 = (*RenderVertex)[Flesh->Indices[i][2]];
 
 									RenderMesh->Vertices.Add(FVector(P1));
 									RenderMesh->Vertices.Add(FVector(P2));
@@ -476,9 +524,9 @@ void UFleshComponent::RenderProceduralMesh()
 									RenderMesh->UVs.Add(FVector2D(0, 0));
 									RenderMesh->UVs.Add(FVector2D(0, 0));
 
-									RenderMesh->Triangles.Add(3 * i);
-									RenderMesh->Triangles.Add(3 * i + 1);
-									RenderMesh->Triangles.Add(3 * i + 2);
+									RenderMesh->Triangles.Add(3 * j);
+									RenderMesh->Triangles.Add(3 * j + 1);
+									RenderMesh->Triangles.Add(3 * j + 2);
 
 									auto Normal = -Chaos::FVec3::CrossProduct(P3 - P1, P2 - P1);
 									RenderMesh->Normals.Add(Normal);
@@ -491,6 +539,8 @@ void UFleshComponent::RenderProceduralMesh()
 									RenderMesh->Tangents.Add(FProcMeshTangent(Tangent[0], Tangent[1], Tangent[2]));
 									Tangent = (P1 - P3).GetSafeNormal();
 									RenderMesh->Tangents.Add(FProcMeshTangent(Tangent[0], Tangent[1], Tangent[2]));
+								
+									j++;
 								}
 
 								Mesh->SetRelativeTransform(GetComponentTransform());
@@ -498,37 +548,40 @@ void UFleshComponent::RenderProceduralMesh()
 							}
 							else
 							{
-
-								const TManagedArray<FVector3f>* RenderVertex = &Flesh->Vertex;
-								if (GetDynamicCollection())
-								{
-									const TManagedArray<FVector3f>& DynamicVertex = GetDynamicCollection()->GetPositions();
-									if (DynamicVertex.Num()) RenderVertex = &DynamicVertex;
-								}
 								auto InRange = [](int32 Size, int32 Val) { return 0 <= Val && Val < Size; };
 
 								// Display only
-								for (int i = 0; i < NumFaces; ++i)
+								for (int i = 0, j = 0; i < NumFaces; ++i)
 								{
+									if (((FacesVisible && (*FacesVisible)[i] == false)) ||
+										(FaceVerticesToSkip.Contains(Flesh->Indices[i][0]) &&
+										 FaceVerticesToSkip.Contains(Flesh->Indices[i][1]) &&
+										 FaceVerticesToSkip.Contains(Flesh->Indices[i][2])))
+									{
+										continue;
+									}
+
 									const auto& P1 = (*RenderVertex)[Flesh->Indices[i][0]];
 									const auto& P2 = (*RenderVertex)[Flesh->Indices[i][1]];
 									const auto& P3 = (*RenderVertex)[Flesh->Indices[i][2]];
 
-									RenderMesh->Vertices[3 * i] = FVector(P1);
-									RenderMesh->Vertices[3 * i + 1] = FVector(P2);
-									RenderMesh->Vertices[3 * i + 2] = FVector(P3);
+									RenderMesh->Vertices[3 * j] = FVector(P1);
+									RenderMesh->Vertices[3 * j + 1] = FVector(P2);
+									RenderMesh->Vertices[3 * j + 2] = FVector(P3);
 
 									auto Normal = Chaos::FVec3::CrossProduct(P3 - P1, P2 - P1);
-									RenderMesh->Normals[3 * i] = Normal;
-									RenderMesh->Normals[3 * i + 1] = Normal;
-									RenderMesh->Normals[3 * i + 2] = Normal;
+									RenderMesh->Normals[3 * j] = Normal;
+									RenderMesh->Normals[3 * j + 1] = Normal;
+									RenderMesh->Normals[3 * j + 2] = Normal;
 
 									auto Tangent = (P2 - P1).GetSafeNormal();
-									RenderMesh->Tangents[3 * i] = FProcMeshTangent(Tangent[0], Tangent[1], Tangent[2]);
+									RenderMesh->Tangents[3 * j] = FProcMeshTangent(Tangent[0], Tangent[1], Tangent[2]);
 									Tangent = (P3 - P2).GetSafeNormal();
-									RenderMesh->Tangents[3 * i + 1] = FProcMeshTangent(Tangent[0], Tangent[1], Tangent[2]);
+									RenderMesh->Tangents[3 * j + 1] = FProcMeshTangent(Tangent[0], Tangent[1], Tangent[2]);
 									Tangent = (P1 - P3).GetSafeNormal();
-									RenderMesh->Tangents[3 * i + 2] = FProcMeshTangent(Tangent[0], Tangent[1], Tangent[2]);
+									RenderMesh->Tangents[3 * j + 2] = FProcMeshTangent(Tangent[0], Tangent[1], Tangent[2]);
+
+									j++;
 								}
 
 								if (!Mesh->GetComponentTransform().Equals(GetComponentTransform())) {
diff --git a/Engine/Plugins/Experimental/ChaosFlesh/Source/ChaosFleshEngine/Public/ChaosFlesh/FleshComponent.h b/Engine/Plugins/Experimental/ChaosFlesh/Source/ChaosFleshEngine/Public/ChaosFlesh/FleshComponent.h
--- a/Engine/Plugins/Experimental/ChaosFlesh/Source/ChaosFleshEngine/Public/ChaosFlesh/FleshComponent.h
+++ b/Engine/Plugins/Experimental/ChaosFlesh/Source/ChaosFleshEngine/Public/ChaosFlesh/FleshComponent.h
@@ -92,11 +92,9 @@ public:
 	UFUNCTION(BlueprintCallable, Category = "Physics")
 	TArray<FVector> GetSkeletalMeshEmbeddedPositions(const ChaosDeformableBindingOption Format, const FTransform TargetDeformationSkeletonOffset, const FName TargetBone = "", const float SimulationBlendWeight = 1.f) const;
 
-	/** 
-	* Offset transform that moves the \c TargetDeformationSkeleton to be co-located with the flesh mesh. 
-	*/
-//	UPROPERTY(EditAnywhere, BlueprintReadWrite, Category = "Rendering")
-//	FTransform TargetDeformationSkeletonOffset;
+	/** Indices of tetrahedra to hide. */
+	UPROPERTY(EditAnywhere, Category = "Rendering")
+	TSet<int32> HideTetrahedra;
 
 private:
 	/** FleshAsset that describes the simulation rest state. */
@@ -172,5 +170,8 @@ private:
 	int32 SimSpaceTransformIndex = INDEX_NONE;
 	int32 SimSpaceTransformGlobalIndex = INDEX_NONE;
 
+	TSet<int32> FaceVerticesToSkip;
+	int32 NumSkippedFaces = 0;
+	int32 NumSkippedTets = 0;
 };
 
diff --git a/Engine/Plugins/Runtime/ProceduralMeshComponent/Source/ProceduralMeshComponent/Private/ProceduralMeshComponent.cpp b/Engine/Plugins/Runtime/ProceduralMeshComponent/Source/ProceduralMeshComponent/Private/ProceduralMeshComponent.cpp
--- a/Engine/Plugins/Runtime/ProceduralMeshComponent/Source/ProceduralMeshComponent/Private/ProceduralMeshComponent.cpp
+++ b/Engine/Plugins/Runtime/ProceduralMeshComponent/Source/ProceduralMeshComponent/Private/ProceduralMeshComponent.cpp
@@ -689,7 +689,7 @@ void UProceduralMeshComponent::UpdateMeshSection(int32 SectionIndex, const TArra
 		// Update bounds, if we are getting new position data
 		if (bSameVertexCount)
 		{
-			Section.SectionLocalBox = FBox(Vertices);
+			Section.SectionLocalBox = Vertices.Num() ? FBox(Vertices) : FBox(EForceInit::ForceInit);
 
 			// Iterate through vertex data, copying in new info
 			for (int32 VertIdx = 0; VertIdx < NumVerts; VertIdx++)
