ChaosFlesh : Fix Raycasts collisions to respect the channel configurations for hit geometry.

diff --git a/Engine/Plugins/Experimental/ChaosFlesh/Source/ChaosFleshEngine/Private/ChaosFlesh/ChaosDeformableGameplayComponent.cpp b/Engine/Plugins/Experimental/ChaosFlesh/Source/ChaosFleshEngine/Private/ChaosFlesh/ChaosDeformableGameplayComponent.cpp
--- a/Engine/Plugins/Experimental/ChaosFlesh/Source/ChaosFleshEngine/Private/ChaosFlesh/ChaosDeformableGameplayComponent.cpp
+++ b/Engine/Plugins/Experimental/ChaosFlesh/Source/ChaosFleshEngine/Private/ChaosFlesh/ChaosDeformableGameplayComponent.cpp
@@ -300,13 +300,7 @@ void UDeformableGameplayComponent::DetectEnvironmentCollisions(const int32 MaxNu
 										   const int32 TetVertexIndex = Candidates[Index].Index;
 										   const FVector& WorldOrigin = Candidates[Index].Origin;
 										   const FVector& WorldSkinnedPos = Candidates[Index].Target;
-										   HitResult[Index] =
-											   GetWorld()->LineTraceSingleByObjectType(
-											   HitBuffer[Index],
-											   WorldOrigin,
-											   WorldSkinnedPos,
-											   CollisionChannel,
-											   TraceParams);
+										   HitResult[Index] = GetWorld()->LineTraceSingleByChannel(HitBuffer[Index], WorldOrigin, WorldSkinnedPos, CollisionChannel, TraceParams);
 									   }
 								   },
 								   CVarParams2.EnvCollisionsLineTraceBatchSize,
