Optimize flesh and cloth chaos caching. Flesh caches have changed so they are not compatible with old caches. Cloth caches are still compatible.

==CODERS==
FleshComponentCacheAdapter.cpp
Update Flesh to use "channels" instead of per-particle "curves" to store its data as this requires less overhead.

ChaosClothingCacheSchema.cpp
Update Cloth to not use particles to store indices. This same data is covered by the ChannelIndices, but ChannelIndices is faster.

KeyHandle.{h,cpp}, RichCurve.{h,cpp}
Add ability to reserve space for keys.
Detect when keys are being appended and optimize for this case.

ChaosCache.{h,cpp}
Make it faster to flush multiple pending frames at once by doing two passes. The first pass reserves space. The second pass adds the keys. This two pass system is only used when you only have channel data.
Also made it so channels work if you change which particles are cached per frame (e.g., if  switching lods). This path is not as optimal.

==TESTING==
Tested locally cloth and flesh caching with Take Recorder, Sequencer, MRQ

diff --git a/Engine/Plugins/Experimental/ChaosCaching/Source/ChaosCaching/Private/Chaos/ChaosCache.cpp b/Engine/Plugins/Experimental/ChaosCaching/Source/ChaosCaching/Private/Chaos/ChaosCache.cpp
--- a/Engine/Plugins/Experimental/ChaosCaching/Source/ChaosCaching/Private/Chaos/ChaosCache.cpp
+++ b/Engine/Plugins/Experimental/ChaosCaching/Source/ChaosCaching/Private/Chaos/ChaosCache.cpp
@@ -103,12 +103,81 @@ void UChaosCache::PostLoad()
 	// Up-to-date now
 	Version = CurrentVersion;
 
+	// Build ParticleToChannelCurve 
+	ParticleToChannelCurve.Reserve(ChannelCurveToParticle.Num());
+	for (int32 ChannelCurveIndex = 0; ChannelCurveIndex < ChannelCurveToParticle.Num(); ++ChannelCurveIndex)
+	{
+		ParticleToChannelCurve.Add(ChannelCurveToParticle[ChannelCurveIndex], ChannelCurveIndex);
+	}
+
 	Super::PostLoad();
 }
 
-void UChaosCache::FlushPendingFrames()
+void UChaosCache::FlushPendingFrames_ChannelOnlyReservePass(TQueue<FPendingFrameWrite, EQueueMode::Spsc>& LocalPendingWrites,
+	bool& bCanSimpleCopyChannelData)
+{
+	bCanSimpleCopyChannelData = true;
+
+	TMap<FName, int32> ChannelKeys;
+	FPendingFrameWrite NewData;
+	while (PendingWrites.Dequeue(NewData))
+	{
+		QUICK_SCOPE_CYCLE_COUNTER(QSTAT_CacheFlushSingleFrameReservePass);
+
+		const bool bContainsParticleData = NewData.PendingParticleData.Num() > 0;
+		if (!bContainsParticleData)
+		{
+			const int32 NumParticles = NewData.PendingChannelsIndices.Num();
+			for (const TPair<FName, TArray<float>>& ChannelData : NewData.PendingChannelsData)
+			{
+				int32* ChannelKeyCount = ChannelKeys.Find(ChannelData.Key);
+				if (!ChannelKeyCount)
+				{
+					ChannelKeyCount = &ChannelKeys.Add(ChannelData.Key);
+					*ChannelKeyCount = 0;
+				}
+				(*ChannelKeyCount)++;
+			}
+
+			for (int32 PendingIndex = 0; PendingIndex < NumParticles; ++PendingIndex)
+			{
+				const int32 ParticleIndex = NewData.PendingChannelsIndices[PendingIndex];
+				if (bCanSimpleCopyChannelData && PendingIndex < ChannelCurveToParticle.Num())
+				{
+					bCanSimpleCopyChannelData = ChannelCurveToParticle[PendingIndex] == ParticleIndex;
+				}
+
+				if (!ParticleToChannelCurve.Contains(ParticleIndex))
+				{
+					checkSlow(!ChannelCurveToParticle.Contains(ParticleIndex));
+					const int32 ChannelIndex = ChannelCurveToParticle.Add(ParticleIndex);
+					ParticleToChannelCurve.Add(ParticleIndex, ChannelIndex);
+				}
+			}
+		}
+
+		LocalPendingWrites.Enqueue(MoveTemp(NewData));
+	}
+
+	const int32 NumChannels = ChannelCurveToParticle.Num();
+	for (const TPair<FName, int32>& KeyCount : ChannelKeys)
+	{
+		FRichCurves* TargetCurve = ChannelsTracks.Find(KeyCount.Key);
+		if (TargetCurve == nullptr)
+		{
+			TargetCurve = &ChannelsTracks.Add(KeyCount.Key);
+			TargetCurve->RichCurves.SetNum(NumChannels);
+		}
+		for (FRichCurve& Curve : TargetCurve->RichCurves)
+		{
+			Curve.ReserveKeys(Curve.GetNumKeys() + KeyCount.Value);
+		}
+	}
+}
+
+template<EQueueMode Mode>
+bool UChaosCache::FlushPendingFrames_MainPass(TQueue<FPendingFrameWrite, Mode>& InPendingWrites, bool bCanSimpleCopyChannelData)
 {
-	QUICK_SCOPE_CYCLE_COUNTER(QSTAT_CacheFlushPendingFrames);
 	bool bWroteParticleData = false;
 	FPendingFrameWrite NewData;
 
@@ -128,7 +197,7 @@ void UChaosCache::FlushPendingFrames()
 			// Signals that this is the final keyframe and that the particle then deactivates.
 			PTrack.bDeactivateOnEnd = true;
 		}
-		
+
 		if (ensure(PTrack.GetNumKeys() == 0 || Time > PTrack.KeyTimestamps.Last()))
 		{
 			PTrack.KeyTimestamps.Add(Time);
@@ -143,61 +212,78 @@ void UChaosCache::FlushPendingFrames()
 		return false;
 	};
 
-	while(PendingWrites.Dequeue(NewData))
+	while (InPendingWrites.Dequeue(NewData))
 	{
-		QUICK_SCOPE_CYCLE_COUNTER(QSTAT_CacheFlushSingleFrame);
-
+		QUICK_SCOPE_CYCLE_COUNTER(QSTAT_CacheFlushSingleFramePass2);
 		{
-			const int32 NumParticles = NewData.PendingChannelsIndices.Num();
+			const bool bContainsParticleData = NewData.PendingParticleData.Num() > 0;
+			const int32 NumParticles = bContainsParticleData ? NewData.PendingParticleData.Num() : NewData.PendingChannelsIndices.Num();
 			bWroteParticleData |= NumParticles > 0;
-
-			for(auto& ChannelData : NewData.PendingChannelsData)
+			for (auto& ChannelData : NewData.PendingChannelsData)
 			{
 				// TArray<FRichCurve>& TargetCurve = ChannelsTracks.FindOrAdd(ChannelData.Key).RichCurves;
 				// TargetCurve.SetNum(NumParticles);
 
 				FRichCurves* TargetCurve = ChannelsTracks.Find(ChannelData.Key);
-				if(TargetCurve == nullptr)
+				if (TargetCurve == nullptr)
 				{
 					TargetCurve = &ChannelsTracks.Add(ChannelData.Key);
 					TargetCurve->RichCurves.SetNum(NumParticles);
 				}
-				int32 ParticleIndex = 0;
-				for(auto& ChannelValue : ChannelData.Value)
+
+				// If there is no particle data, use ChannelCurveToParticle instead.				
+				check(ChannelData.Value.Num() <= NumParticles);
+				if (bContainsParticleData)
 				{
-					TargetCurve->RichCurves[ParticleIndex++].AddKey(NewData.Time, ChannelValue);
+					for (int32 PendingIndex = 0; PendingIndex < NumParticles; ++PendingIndex)
+					{
+						const FPendingParticleWrite& ParticleData = NewData.PendingParticleData[PendingIndex];
+						const int32 ParticleIndex = ParticleData.ParticleIndex;
+
+						int32 TrackIndex = INDEX_NONE;
+						if (!TrackToParticle.Find(ParticleIndex, TrackIndex))
+						{
+							TrackToParticle.Add(ParticleIndex);
+							TrackIndex = ParticleTracks.AddDefaulted();
+						}
+						TargetCurve->RichCurves[TrackIndex].AddKey(NewData.Time, ChannelData.Value[PendingIndex]);
+					}
 				}
-			}
-
-			const bool bContainsParticleData = NewData.PendingParticleData.Num() > 0;
-			// If there is no particle data, use ChannelCurveToParticle instead. 
-			if (!bContainsParticleData)
-			{
-				// If there are more particles, append the curve index to the particle index map
-				const int32 OldNum = ChannelCurveToParticle.Num();
-				const int32 NewNum = NewData.PendingChannelsIndices.Num();	
-				if (OldNum < NewNum)
+				else if (bCanSimpleCopyChannelData)
+				{
+					int32 ParticleIndex = 0;
+					for (auto& ChannelValue : ChannelData.Value)
+					{
+						TargetCurve->RichCurves[ParticleIndex++].AddKey(NewData.Time, ChannelValue);
+					}
+				}
+				else
 				{
-					ChannelCurveToParticle.SetNum(NewNum);
-					for(int32 Index = OldNum; Index < NewNum; ++Index)
+					for (int32 PendingIndex = 0; PendingIndex < NumParticles; ++PendingIndex)
 					{
-						ChannelCurveToParticle[Index] = NewData.PendingChannelsIndices[Index];
+						const int32 ParticleIndex = NewData.PendingChannelsIndices[PendingIndex];
+						int32* ChannelIndex = ParticleToChannelCurve.Find(ParticleIndex);
+						if (!ChannelIndex)
+						{
+							ChannelIndex = &ParticleToChannelCurve.Add(ParticleIndex, ChannelCurveToParticle.Add(ParticleIndex));
+						}
+						TargetCurve->RichCurves[*ChannelIndex].AddKey(NewData.Time, ChannelData.Value[PendingIndex]);
 					}
 				}
 			}
 		}
-		
+
 		const int32 ParticleCount = NewData.PendingParticleData.Num();
 
 		bWroteParticleData |= ParticleCount > 0;
 
-		for(int32 Index = 0; Index < ParticleCount; ++Index)
+		for (int32 Index = 0; Index < ParticleCount; ++Index)
 		{
 			const FPendingParticleWrite& ParticleData = NewData.PendingParticleData[Index];
 			const int32 ParticleIndex = ParticleData.ParticleIndex;
 
 			int32 TrackIndex = INDEX_NONE;
-			if(!TrackToParticle.Find(ParticleIndex, TrackIndex))
+			if (!TrackToParticle.Find(ParticleIndex, TrackIndex))
 			{
 				TrackToParticle.Add(ParticleIndex);
 				TrackIndex = ParticleTracks.AddDefaulted();
@@ -208,7 +294,7 @@ void UChaosCache::FlushPendingFrames()
 
 			if (WriteDataToTransformTrack(PTrack, NewData.Time, ParticleData.bPendingDeactivate, ParticleData.PendingTransform))
 			{
-				for(TPair<FName, float> CurveKeyPair : ParticleData.PendingCurveData)
+				for (TPair<FName, float> CurveKeyPair : ParticleData.PendingCurveData)
 				{
 					FRichCurve& TargetCurve = TargetCacheData.CurveData.FindOrAdd(CurveKeyPair.Key);
 					TargetCurve.AddKey(NewData.Time, CurveKeyPair.Value);
@@ -216,11 +302,11 @@ void UChaosCache::FlushPendingFrames()
 			}
 		}
 
-		for(TTuple<FName, FCacheEventTrack>& PendingTrack : NewData.PendingEvents)
+		for (TTuple<FName, FCacheEventTrack>& PendingTrack : NewData.PendingEvents)
 		{
 			FCacheEventTrack* CacheTrack = EventTracks.Find(PendingTrack.Key);
 
-			if(!CacheTrack)
+			if (!CacheTrack)
 			{
 				CacheTrack = &EventTracks.Add(PendingTrack.Key, FCacheEventTrack(PendingTrack.Key, PendingTrack.Value.Struct));
 			}
@@ -235,11 +321,32 @@ void UChaosCache::FlushPendingFrames()
 			{
 				TransformCacheTrack = &NamedTransformTracks.Add(PendingNamedTransform.Key, FNamedTransformTrack());
 			}
-			WriteDataToTransformTrack(*TransformCacheTrack, NewData.Time, false, PendingNamedTransform.Value);			
+			WriteDataToTransformTrack(*TransformCacheTrack, NewData.Time, false, PendingNamedTransform.Value);
 		}
 
 		++NumRecordedFrames;
 	}
+	return bWroteParticleData;
+}
+
+void UChaosCache::FlushPendingFrames()
+{
+	QUICK_SCOPE_CYCLE_COUNTER(QSTAT_CacheFlushPendingFrames);
+	bool bWroteParticleData = false;
+
+	const bool bDoChannelOnlyTwoPass = PendingWrites.Peek() && PendingWrites.Peek()->PendingParticleData.Num() == 0 && PendingWrites.Peek()->PendingChannelsIndices.Num() > 0;
+	if (bDoChannelOnlyTwoPass)
+	{
+		bool bCanSimpleCopyChannelData = false;
+		TQueue<FPendingFrameWrite, EQueueMode::Spsc> LocalPendingWrites;
+		FlushPendingFrames_ChannelOnlyReservePass(LocalPendingWrites, bCanSimpleCopyChannelData);
+		bWroteParticleData = FlushPendingFrames_MainPass(LocalPendingWrites, bCanSimpleCopyChannelData);
+	}
+	else
+	{
+		constexpr bool bCanSimpleCopyChannelData = false;
+		bWroteParticleData = FlushPendingFrames_MainPass(PendingWrites, bCanSimpleCopyChannelData);
+	}
 
 	if(bWroteParticleData)
 	{
@@ -291,6 +398,7 @@ FCacheUserToken UChaosCache::BeginRecord(UPrimitiveComponent* InComponent, FGuid
 			CurveData.Reset();
 			EventTracks.Reset();
 			NamedTransformTracks.Reset();
+			ParticleToChannelCurve.Reset();
 
 			PendingWrites.Empty();
 
diff --git a/Engine/Plugins/Experimental/ChaosCaching/Source/ChaosCaching/Public/Chaos/ChaosCache.h b/Engine/Plugins/Experimental/ChaosCaching/Source/ChaosCaching/Public/Chaos/ChaosCache.h
--- a/Engine/Plugins/Experimental/ChaosCaching/Source/ChaosCaching/Public/Chaos/ChaosCache.h
+++ b/Engine/Plugins/Experimental/ChaosCaching/Source/ChaosCaching/Public/Chaos/ChaosCache.h
@@ -437,6 +437,13 @@ public:
 	}
 
 private:
+	void FlushPendingFrames_ChannelOnlyReservePass(TQueue<FPendingFrameWrite, EQueueMode::Spsc>& LocalPendingWrites,
+		bool& bCanSimpleCopyChannelData);
+
+	// @return Whether or not any particle data was written
+	template<EQueueMode Mode>
+	bool FlushPendingFrames_MainPass(TQueue<FPendingFrameWrite, Mode>& InPendingWrites, bool bCanSimpleCopyChannelData);
+
 	friend class AChaosCacheManager;
 
 	/** Timestamped generic event tracks */
@@ -468,4 +475,7 @@ private:
 
 	/** Indicates that we need to strip MassToLocal before playing the cache. */
 	bool bStripMassToLocal;
+
+	/** Reverse Lookup for ChannelCurveToParticle. Rebuilt on load.*/
+	TMap<int32,int32> ParticleToChannelCurve;
 };
diff --git a/Engine/Plugins/Experimental/ChaosCloth/Source/Chaos/Private/ChaosCloth/ChaosClothingCacheSchema.cpp b/Engine/Plugins/Experimental/ChaosCloth/Source/Chaos/Private/ChaosCloth/ChaosClothingCacheSchema.cpp
--- a/Engine/Plugins/Experimental/ChaosCloth/Source/Chaos/Private/ChaosCloth/ChaosClothingCacheSchema.cpp
+++ b/Engine/Plugins/Experimental/ChaosCloth/Source/Chaos/Private/ChaosCloth/ChaosClothingCacheSchema.cpp
@@ -46,11 +46,6 @@ void FClothingCacheSchema::RecordPostSolve(const FClothingSimulationSolver& Clot
 				const Softs::FSolverVec3& ParticleX = ParticleXs[ParticleIndex];
 
 				// Adding the vertices relative position to the particle write datas
-				FPendingParticleWrite NewData;
-
-				NewData.ParticleIndex = ParticleIndex;
-				OutFrame.PendingParticleData.Add(MoveTemp(NewData));
-
 				PendingID.Add(ParticleIndex);
 				PendingVX.Add(ParticleV.X);
 				PendingVY.Add(ParticleV.Y);
@@ -89,15 +84,15 @@ void FClothingCacheSchema::PlaybackPreSolve(UChaosCache& InCache, FReal InTime,
 	const TArray<float>* PendingPX = EvaluatedResult.Channels.Find(PositionXName);
 	const TArray<float>* PendingPY = EvaluatedResult.Channels.Find(PositionYName);
 	const TArray<float>* PendingPZ = EvaluatedResult.Channels.Find(PositionZName);
+	const int32 NumCachedParticles = EvaluatedResult.ParticleIndices.Num();
 
-	if (PendingVX && PendingVY && PendingVZ && PendingPX && PendingPY && PendingPZ)
+	if (PendingVX && PendingVY && PendingVZ && PendingPX && PendingPY && PendingPZ && NumCachedParticles)
 	{
 		// Directly set the result of the cache into the solver particles
-		const int32 NumCachedParticles = EvaluatedResult.ParticleIndices.Num();
 		Softs::FSolverVec3* ParticleXs = ClothSolver.GetParticleXs(0);
 		Softs::FSolverVec3* ParticleVs = ClothSolver.GetParticleVs(0);
 		const int32 NumParticles = ClothSolver.GetNumParticles();
-		for (int32 CachedIndex = 0; CachedIndex < NumParticles; ++CachedIndex)
+		for (int32 CachedIndex = 0; CachedIndex < NumCachedParticles; ++CachedIndex)
 		{
 			const int32 ParticleIndex = EvaluatedResult.ParticleIndices[CachedIndex];
 			check(ParticleIndex < NumParticles);
diff --git a/Engine/Plugins/Experimental/ChaosFlesh/Source/ChaosFleshEngine/Private/ChaosCache/FleshComponentCacheAdapter.cpp b/Engine/Plugins/Experimental/ChaosFlesh/Source/ChaosFleshEngine/Private/ChaosCache/FleshComponentCacheAdapter.cpp
--- a/Engine/Plugins/Experimental/ChaosFlesh/Source/ChaosFleshEngine/Private/ChaosCache/FleshComponentCacheAdapter.cpp
+++ b/Engine/Plugins/Experimental/ChaosFlesh/Source/ChaosFleshEngine/Private/ChaosCache/FleshComponentCacheAdapter.cpp
@@ -52,27 +52,39 @@ namespace Chaos
 					const Softs::FSolverVec3* ParticleXs = &Particles.X(0);
 					const Softs::FSolverVec3* ParticleVs = &Particles.V(0);
 
+					TArray<float> PendingVX, PendingVY, PendingVZ, PendingPX, PendingPY, PendingPZ;
+					TArray<int32>& PendingID = OutFrame.PendingChannelsIndices;
+
+					PendingID.Reserve(NumParticles);
+					PendingVX.Reserve(NumParticles);
+					PendingVY.Reserve(NumParticles);
+					PendingVZ.Reserve(NumParticles);
+					PendingPX.Reserve(NumParticles);
+					PendingPY.Reserve(NumParticles);
+					PendingPZ.Reserve(NumParticles);
+
 					for (uint32 ParticleIndex = 0; ParticleIndex < NumParticles; ++ParticleIndex)
 					{
 						const Softs::FSolverVec3& ParticleV = ParticleVs[ParticleIndex];
 						const Softs::FSolverVec3& ParticleX = ParticleXs[ParticleIndex];
 
 						// Adding the vertices relative position to the particle write datas
-						FPendingParticleWrite NewData;
-
-						NewData.ParticleIndex = ParticleIndex;
-						//	NewData.PendingTransform = FTransform(FQuat::Identity, FVector(ParticleX)).GetRelativeTransform(InRootTransform);
-
-						NewData.PendingCurveData.Add(MakeTuple(VelocityXName, ParticleV.X));
-						NewData.PendingCurveData.Add(MakeTuple(VelocityYName, ParticleV.Y));
-						NewData.PendingCurveData.Add(MakeTuple(VelocityZName, ParticleV.Z));
-
-						NewData.PendingCurveData.Add(MakeTuple(PositionXName, ParticleX.X));
-						NewData.PendingCurveData.Add(MakeTuple(PositionYName, ParticleX.Y));
-						NewData.PendingCurveData.Add(MakeTuple(PositionZName, ParticleX.Z));
-
-						OutFrame.PendingParticleData.Add(MoveTemp(NewData));
+						PendingID.Add(ParticleIndex);
+						PendingVX.Add(ParticleV.X);
+						PendingVY.Add(ParticleV.Y);
+						PendingVZ.Add(ParticleV.Z);
+
+						PendingPX.Add(ParticleX.X);
+						PendingPY.Add(ParticleX.Y);
+						PendingPZ.Add(ParticleX.Z);
 					}
+
+					OutFrame.PendingChannelsData.Add(VelocityXName, PendingVX);
+					OutFrame.PendingChannelsData.Add(VelocityYName, PendingVY);
+					OutFrame.PendingChannelsData.Add(VelocityZName, PendingVZ);
+					OutFrame.PendingChannelsData.Add(PositionXName, PendingPX);
+					OutFrame.PendingChannelsData.Add(PositionYName, PendingPY);
+					OutFrame.PendingChannelsData.Add(PositionZName, PendingPZ);
 				}
 			}
 		}
@@ -88,57 +100,45 @@ namespace Chaos
 			{
 				FCacheEvaluationContext Context(TickRecord);
 				Context.bEvaluateTransform = false;
-				Context.bEvaluateCurves = true;
+				Context.bEvaluateCurves = false;
 				Context.bEvaluateEvents = false;
+				Context.bEvaluateChannels = true;
 
 				// The evaluated result are already in world space since we are passing the tickrecord.spacetransform in the context
 				FCacheEvaluationResult EvaluatedResult = InCache->Evaluate(Context, nullptr);
-				const int32 NumCurves = EvaluatedResult.Curves.Num();
+
+				const TArray<float>* PendingVX = EvaluatedResult.Channels.Find(VelocityXName);
+				const TArray<float>* PendingVY = EvaluatedResult.Channels.Find(VelocityYName);
+				const TArray<float>* PendingVZ = EvaluatedResult.Channels.Find(VelocityZName);
+				const TArray<float>* PendingPX = EvaluatedResult.Channels.Find(PositionXName);
+				const TArray<float>* PendingPY = EvaluatedResult.Channels.Find(PositionYName);
+				const TArray<float>* PendingPZ = EvaluatedResult.Channels.Find(PositionZName);
+
+				const int32 NumCachedParticles = EvaluatedResult.ParticleIndices.Num();
 
 				FParticles& Particles = Evolution->Particles();
 
-				const uint32 NumParticles = Particles.Size();
-				if (NumParticles > 0)
+				const int32 NumParticles = Particles.Size();
+				if (NumCachedParticles > 0)
 				{
+					// Directly set the result of the cache into the solver particles
+
 					Softs::FSolverVec3* ParticleXs = &Particles.X(0);
 					Softs::FSolverVec3* ParticleVs = &Particles.V(0);
 
-					for (uint32 ParticleIndex = 0; ParticleIndex < NumParticles; ++ParticleIndex)
+					for (int32 CachedIndex = 0; CachedIndex < NumCachedParticles; ++CachedIndex)
 					{
-						Softs::FSolverVec3& ParticleV = ParticleVs[ParticleIndex];
-						Softs::FSolverVec3& ParticleX = ParticleXs[ParticleIndex];
-
-						// Directly set the result of the cache into the solver particles
-						if ( ParticleIndex < (uint32)EvaluatedResult.ParticleIndices.Num())
+						const int32 ParticleIndex = EvaluatedResult.ParticleIndices[CachedIndex];
+						if (ensure(ParticleIndex < NumParticles))
 						{
-							const int32 TransformIndex = EvaluatedResult.ParticleIndices[ParticleIndex];
-							if (EvaluatedResult.Curves.IsValidIndex(TransformIndex))
-							{
-								if (const float* VelocityX = EvaluatedResult.Curves[TransformIndex].Find(VelocityXName))
-								{
-									ParticleV.X = *VelocityX;
-								}
-								if (const float* VelocityY = EvaluatedResult.Curves[TransformIndex].Find(VelocityYName))
-								{
-									ParticleV.Y = *VelocityY;
-								}
-								if (const float* VelocityZ = EvaluatedResult.Curves[TransformIndex].Find(VelocityZName))
-								{
-									ParticleV.Z = *VelocityZ;
-								}
-								if (const float* PositionX = EvaluatedResult.Curves[TransformIndex].Find(PositionXName))
-								{
-									ParticleX.X = *PositionX;
-								}
-								if (const float* PositionY = EvaluatedResult.Curves[TransformIndex].Find(PositionYName))
-								{
-									ParticleX.Y = *PositionY;
-								}
-								if (const float* PositionZ = EvaluatedResult.Curves[TransformIndex].Find(PositionZName))
-								{
-									ParticleX.Z = *PositionZ;
-								}
-							}
+							Softs::FSolverVec3& ParticleV = ParticleVs[ParticleIndex];
+							Softs::FSolverVec3& ParticleX = ParticleXs[ParticleIndex];
+							ParticleV.X = (*PendingVX)[CachedIndex];
+							ParticleV.Y = (*PendingVY)[CachedIndex];
+							ParticleV.Z = (*PendingVZ)[CachedIndex];
+							ParticleX.X = (*PendingPX)[CachedIndex];
+							ParticleX.Y = (*PendingPY)[CachedIndex];
+							ParticleX.Z = (*PendingPZ)[CachedIndex];
 						}
 					}
 				}
@@ -157,7 +157,7 @@ namespace Chaos
 	{
 		// If we have a flesh mesh we can play back any cache as long as it has one or more tracks
 		const UFleshComponent* FleshComp = CastChecked<UFleshComponent>(InComponent);
-		return FleshComp && InCache->TrackToParticle.Num() > 0;
+		return FleshComp && InCache->ChannelCurveToParticle.Num() > 0;
 	}
 
 	Chaos::Softs::FDeformableSolver* FFleshCacheAdapter::GetDeformableSolver(UPrimitiveComponent* InComponent) const
@@ -208,59 +208,49 @@ namespace Chaos
 
 			FCacheEvaluationContext Context(TickRecord);
 			Context.bEvaluateTransform = false;
-			Context.bEvaluateCurves = true;
+			Context.bEvaluateCurves = false;
 			Context.bEvaluateEvents = false;
+			Context.bEvaluateChannels = true;
 
 			FCacheEvaluationResult EvaluatedResult = InCache->Evaluate(Context, nullptr);
 
+			const TArray<float>* PendingVX = EvaluatedResult.Channels.Find(VelocityXName);
+			const TArray<float>* PendingVY = EvaluatedResult.Channels.Find(VelocityYName);
+			const TArray<float>* PendingVZ = EvaluatedResult.Channels.Find(VelocityZName);
+			const TArray<float>* PendingPX = EvaluatedResult.Channels.Find(PositionXName);
+			const TArray<float>* PendingPY = EvaluatedResult.Channels.Find(PositionYName);
+			const TArray<float>* PendingPZ = EvaluatedResult.Channels.Find(PositionZName);
+
+			const int32 NumCachedParticles = EvaluatedResult.ParticleIndices.Num();
+
+			const bool bHasPositions = PendingPX && PendingPY && PendingPZ;
+			const bool bHasVelocities = PendingVX && PendingVY && PendingVZ;
+
+
 			FleshComp->ResetDynamicCollection();
-			if (UFleshDynamicAsset* DynamicCollection = FleshComp->GetDynamicCollection())
-			{
-				TManagedArray<FVector3f>& DynamicVertex = DynamicCollection->GetPositions();
-				//TManagedArray<FVector3f>& DynamicVertex = DynamicCollection->GetVelocities();
 
-				const int32 NumCurves = EvaluatedResult.Curves.Num();
-				if (NumCurves == EvaluatedResult.ParticleIndices.Num())
+			if (bHasPositions && NumCachedParticles > 0)
+			{
+				if (UFleshDynamicAsset* DynamicCollection = FleshComp->GetDynamicCollection())
 				{
-					if (DynamicVertex.Num() == NumCurves)
+					TManagedArray<FVector3f>& DynamicVertex = DynamicCollection->GetPositions();
+					//TManagedArray<FVector3f>& DynamicVertex = DynamicCollection->GetVelocities();
+					const int32 NumDynamicVertex = DynamicVertex.Num();
+					if (NumDynamicVertex == NumCachedParticles)
 					{
-						for (int32 ParticleIndex = 0; ParticleIndex < NumCurves; ++ParticleIndex)
+						for (int32 CachedIndex = 0; CachedIndex < NumCachedParticles; ++CachedIndex)
 						{
-							if (ParticleIndex < EvaluatedResult.ParticleIndices.Num())
+							const int32 ParticleIndex = EvaluatedResult.ParticleIndices[CachedIndex];
+
+							if (ensure(ParticleIndex < NumDynamicVertex))
 							{
-								const int32 TransformIndex = EvaluatedResult.ParticleIndices[ParticleIndex];
-								if (EvaluatedResult.Curves.IsValidIndex(TransformIndex))
-								{
-									/*
-									if (const float* VelocityX = EvaluatedResult.Curves[TransformIndex].Find(VelocityXName))
-									{
-										SimulationContext->CachedVelocities[ParticleIndex].X = *VelocityX;
-									}
-									if (const float* VelocityY = EvaluatedResult.Curves[TransformIndex].Find(VelocityYName))
-									{
-										SimulationContext->CachedVelocities[ParticleIndex].Y = *VelocityY;
-									}
-									if (const float* VelocityZ = EvaluatedResult.Curves[TransformIndex].Find(VelocityZName))
-									{
-										SimulationContext->CachedVelocities[ParticleIndex].Z = *VelocityZ;
-									}
-									*/
-									if (const float* PositionX = EvaluatedResult.Curves[TransformIndex].Find(PositionXName))
-									{
-										DynamicVertex[ParticleIndex].X = *PositionX;
-									}
-									if (const float* PositionY = EvaluatedResult.Curves[TransformIndex].Find(PositionYName))
-									{
-										DynamicVertex[ParticleIndex].Y = *PositionY;
-									}
-									if (const float* PositionZ = EvaluatedResult.Curves[TransformIndex].Find(PositionZName))
-									{
-										DynamicVertex[ParticleIndex].Z = *PositionZ;
-									}
-									auto UEVertd = [](Chaos::FVec3 V) { return FVector3d(V.X, V.Y, V.Z); };
-									auto UEVertf = [](FVector3d V) { return FVector3f((float)V.X, (float)V.Y, (float)V.Z); };
-									DynamicVertex[ParticleIndex] = UEVertf(FleshComp->GetComponentTransform().InverseTransformPosition(UEVertd(DynamicVertex[ParticleIndex])));
-								}
+								DynamicVertex[ParticleIndex].X = (*PendingPX)[CachedIndex];
+								DynamicVertex[ParticleIndex].Y = (*PendingPY)[CachedIndex];
+								DynamicVertex[ParticleIndex].Z = (*PendingPZ)[CachedIndex];
+
+								auto UEVertd = [](Chaos::FVec3 V) { return FVector3d(V.X, V.Y, V.Z); };
+								auto UEVertf = [](FVector3d V) { return FVector3f((float)V.X, (float)V.Y, (float)V.Z); };
+								DynamicVertex[ParticleIndex] = UEVertf(FleshComp->GetComponentTransform().InverseTransformPosition(UEVertd(DynamicVertex[ParticleIndex])));
 							}
 						}
 					}
diff --git a/Engine/Source/Runtime/Engine/Classes/Curves/KeyHandle.h b/Engine/Source/Runtime/Engine/Classes/Curves/KeyHandle.h
--- a/Engine/Source/Runtime/Engine/Classes/Curves/KeyHandle.h
+++ b/Engine/Source/Runtime/Engine/Classes/Curves/KeyHandle.h
@@ -78,6 +78,7 @@ public:
 	void Add( const FKeyHandle& InHandle, int32 InIndex );
 	void Empty(int32 ExpectedNumElements = 0);
 	void Remove( const FKeyHandle& InHandle );
+	void Reserve(int32 NumElements);
 	const int32* Find(const FKeyHandle& InHandle) const { return KeyHandlesToIndices.Find(InHandle); }
 	const FKeyHandle* FindKey( int32 KeyIndex ) const;
 	int32 Num() const { return KeyHandlesToIndices.Num(); }
diff --git a/Engine/Source/Runtime/Engine/Classes/Curves/RichCurve.h b/Engine/Source/Runtime/Engine/Classes/Curves/RichCurve.h
--- a/Engine/Source/Runtime/Engine/Classes/Curves/RichCurve.h
+++ b/Engine/Source/Runtime/Engine/Classes/Curves/RichCurve.h
@@ -227,6 +227,13 @@ public:
 	  */
 	virtual FKeyHandle AddKey(float InTime, float InValue, const bool bUnwindRotation = false, FKeyHandle KeyHandle = FKeyHandle()) final override;
 
+	/**
+	* Reserves keys to be added by AddKey.
+	* 
+	* @see AddKey
+	*/
+	void ReserveKeys(const int32 Number);
+
 	/**
 	 * Sets the keys with the keys.
 	 *
diff --git a/Engine/Source/Runtime/Engine/Private/Curves/KeyHandle.cpp b/Engine/Source/Runtime/Engine/Private/Curves/KeyHandle.cpp
--- a/Engine/Source/Runtime/Engine/Private/Curves/KeyHandle.cpp
+++ b/Engine/Source/Runtime/Engine/Private/Curves/KeyHandle.cpp
@@ -50,11 +50,6 @@ void FKeyHandleMap::Initialize(TArrayView<const FKeyHandle> InKeyHandles)
 
 void FKeyHandleMap::Add( const FKeyHandle& InHandle, int32 InIndex )
 {
-	for (auto It = KeyHandlesToIndices.CreateIterator(); It; ++It)
-	{
-		int32& KeyIndex = It.Value();
-		if (KeyIndex >= InIndex) { ++KeyIndex; }
-	}
 
 	if (InIndex > KeyHandles.Num())
 	{
@@ -68,6 +63,14 @@ void FKeyHandleMap::Add( const FKeyHandle& InHandle, int32 InIndex )
 	}
 	else
 	{
+		if (InIndex < KeyHandles.Num())
+		{
+			for (auto It = KeyHandlesToIndices.CreateIterator(); It; ++It)
+			{
+				int32& KeyIndex = It.Value();
+				if (KeyIndex >= InIndex) { ++KeyIndex; }
+			}
+		}
 		KeyHandles.Insert(InHandle, InIndex);
 	}
 
@@ -95,6 +98,13 @@ void FKeyHandleMap::Empty(int32 ExpectedNumElements)
 }
 
 
+void FKeyHandleMap::Reserve(int32 NumElements)
+{
+	KeyHandlesToIndices.Reserve(NumElements);
+	KeyHandles.Reserve(NumElements);
+}
+
+
 void FKeyHandleMap::Remove( const FKeyHandle& InHandle )
 {
 	int32 Index = INDEX_NONE;
diff --git a/Engine/Source/Runtime/Engine/Private/Curves/RichCurve.cpp b/Engine/Source/Runtime/Engine/Private/Curves/RichCurve.cpp
--- a/Engine/Source/Runtime/Engine/Private/Curves/RichCurve.cpp
+++ b/Engine/Source/Runtime/Engine/Private/Curves/RichCurve.cpp
@@ -301,7 +301,14 @@ FRichCurveKey* FRichCurve::GetFirstMatchingKey(const TArray<FKeyHandle>& KeyHand
 FKeyHandle FRichCurve::AddKey( const float InTime, const float InValue, const bool bUnwindRotation, FKeyHandle NewHandle )
 {
 	int32 Index = 0;
-	for(; Index < Keys.Num() && Keys[Index].Time < InTime; ++Index);
+	if (Keys.Num() == 0 || Keys.Last().Time < InTime)
+	{
+		Index = Keys.Num();
+	}
+	else
+	{
+		for (; Index < Keys.Num() && Keys[Index].Time < InTime; ++Index);
+	}
 	Keys.Insert(FRichCurveKey(InTime, InValue), Index);
 
 	// If we were asked to treat this curve as a rotation value and to unwindow the rotation, then
@@ -330,6 +337,13 @@ FKeyHandle FRichCurve::AddKey( const float InTime, const float InValue, const bo
 }
 
 
+void FRichCurve::ReserveKeys(const int32 Number)
+{
+	Keys.Reserve(Number);
+	KeyHandlesToIndices.Reserve(Number);
+}
+
+
 void FRichCurve::SetKeys(const TArray<FRichCurveKey>& InKeys)
 {
 	Reset();
