ChaosFlesh : Add flesh skeleton binding node. Add threshold of 1 to kinematic weights

diff --git a/Engine/Plugins/Experimental/ChaosFlesh/Source/ChaosFleshNodes/Private/Dataflow/ChaosFleshPositionTargetInitializationNodes.cpp b/Engine/Plugins/Experimental/ChaosFlesh/Source/ChaosFleshNodes/Private/Dataflow/ChaosFleshPositionTargetInitializationNodes.cpp
--- a/Engine/Plugins/Experimental/ChaosFlesh/Source/ChaosFleshNodes/Private/Dataflow/ChaosFleshPositionTargetInitializationNodes.cpp
+++ b/Engine/Plugins/Experimental/ChaosFlesh/Source/ChaosFleshNodes/Private/Dataflow/ChaosFleshPositionTargetInitializationNodes.cpp
@@ -15,6 +15,12 @@
 #include "GeometryCollection/TransformCollection.h"
 #include "Chaos/Tetrahedron.h"
 #include "Chaos/Plane.h"
+#include "Rendering/SkeletalMeshModel.h"
+#include "Rendering/SkeletalMeshRenderData.h"
+#include "Components/SkeletalMeshComponent.h"
+#include "MeshDescription.h"
+#include "MeshDescriptionToDynamicMesh.h"
+#include "BoneWeights.h"
 
 #include UE_INLINE_GENERATED_CPP_BY_NAME(ChaosFleshPositionTargetInitializationNodes)
 
@@ -28,6 +34,7 @@ namespace Dataflow
 		DATAFLOW_NODE_REGISTER_CREATION_FACTORY(FSetVertexVertexPositionTargetBindingDataflowNode);
 		DATAFLOW_NODE_REGISTER_CREATION_FACTORY(FSetVertexTetrahedraPositionTargetBindingDataflowNode);
 		DATAFLOW_NODE_REGISTER_CREATION_FACTORY(FSetVertexTrianglePositionTargetBindingDataflowNode);
+		DATAFLOW_NODE_REGISTER_CREATION_FACTORY(FSetFleshBonePositionTargetBindingDataflowNode);
 	}
 }
 
@@ -539,4 +546,335 @@ void FSetVertexTrianglePositionTargetBindingDataflowNode::Evaluate(Dataflow::FCo
 		}
 		SetValue<DataType>(Context, InCollection, &Collection);
 	}
+}
+
+void FSetFleshBonePositionTargetBindingDataflowNode::Evaluate(Dataflow::FContext& Context, const FDataflowOutput* Out) const
+{
+	if (Out->IsA<DataType>(&Collection))
+	{
+		DataType InCollection = GetValue<DataType>(Context, &Collection);
+
+		if (TObjectPtr<USkeletalMesh> BoneSkeletalMesh = GetValue<TObjectPtr<USkeletalMesh>>(Context, &SkeletalMeshIn))
+		{
+			if (TManagedArray<FIntVector>* Indices = InCollection.FindAttribute<FIntVector>("Indices", FGeometryCollection::FacesGroup))
+			{
+				if (TManagedArray<FVector3f>* Vertices = InCollection.FindAttribute<FVector3f>("Vertex", FGeometryCollection::VerticesGroup))
+				{
+
+					FSkeletalMeshRenderData* RenderData = BoneSkeletalMesh->GetResourceForRendering();
+					if (RenderData->LODRenderData.Num())
+					{
+						FSkeletalMeshLODRenderData* LODRenderData = &RenderData->LODRenderData[0];
+						const FPositionVertexBuffer& PositionVertexBuffer =
+							LODRenderData->StaticVertexBuffers.PositionVertexBuffer;
+
+						const FSkinWeightVertexBuffer* SkinWeightVertexBuffer = LODRenderData->GetSkinWeightVertexBuffer();
+						const int32 MaxBoneInfluences = SkinWeightVertexBuffer->GetMaxBoneInfluences();
+						TArray<FTransform> ComponentPose;
+						Dataflow::Animation::GlobalTransforms(BoneSkeletalMesh->GetRefSkeleton(), ComponentPose);
+						TArray<TArray<int32>> BoneBoundVerts;
+						TArray<TArray<float>> BoneBoundWeights;
+						BoneBoundVerts.SetNum(ComponentPose.Num());
+						BoneBoundWeights.SetNum(ComponentPose.Num());
+
+						uint32 NumSkeletonVertices = PositionVertexBuffer.GetNumVertices();
+
+						Chaos::FReal SphereRadius = (Chaos::FReal)0.;
+
+						Chaos::TVec3<float> CoordMaxs(-FLT_MAX);
+						Chaos::TVec3<float> CoordMins(FLT_MAX);
+						for (int32 i = 0; i < int32(NumSkeletonVertices); i++)
+						{
+							for (int32 j = 0; j < 3; j++)
+							{
+								if (PositionVertexBuffer.VertexPosition(i)[j] > CoordMaxs[j])
+								{
+									CoordMaxs[j] = PositionVertexBuffer.VertexPosition(i)[j];
+								}
+								if (PositionVertexBuffer.VertexPosition(i)[j] < CoordMins[j])
+								{
+									CoordMins[j] = PositionVertexBuffer.VertexPosition(i)[j];
+								}
+							}
+						}
+						Chaos::TVec3<float> CoordDiff = (CoordMaxs - CoordMins) * VertexRadiusRatio;
+						SphereRadius = Chaos::FReal(FGenericPlatformMath::Min(CoordDiff[0], FGenericPlatformMath::Min(CoordDiff[1], CoordDiff[2])));
+
+						TArray<Chaos::TSphere<Chaos::FReal, 3>*> VertexSpherePtrs;
+						TArray<Chaos::TSphere<Chaos::FReal, 3>> VertexSpheres;
+
+						VertexSpheres.Init(Chaos::TSphere<Chaos::FReal, 3>(Chaos::TVec3<Chaos::FReal>(0), SphereRadius), NumSkeletonVertices);
+						VertexSpherePtrs.SetNum(NumSkeletonVertices);
+
+						for (int32 i = 0; i < int32(NumSkeletonVertices); i++)
+						{
+							Chaos::TVec3<Chaos::FReal> SphereCenter(PositionVertexBuffer.VertexPosition(i));
+							Chaos::TSphere<Chaos::FReal, 3> VertexSphere(SphereCenter, SphereRadius);
+							VertexSpheres[i] = Chaos::TSphere<Chaos::FReal, 3>(SphereCenter, SphereRadius);
+							VertexSpherePtrs[i] = &VertexSpheres[i];
+						}
+						Chaos::TBoundingVolumeHierarchy<
+							TArray<Chaos::TSphere<Chaos::FReal, 3>*>,
+							TArray<int32>,
+							Chaos::FReal,
+							3> VertexBVH(VertexSpherePtrs);
+
+						if (SkeletalBindingMode == ESkeletalBindingMode::Dataflow_SkeletalBinding_Kinematic)
+						{
+
+							//avoid particles that are already kinematic:
+							TArray<bool> ParticleIsKinematic;
+							ParticleIsKinematic.Init(false, Vertices->Num());
+							typedef GeometryCollection::Facades::FKinematicBindingFacade FKinematics;
+							FKinematics Kinematics(InCollection);
+
+							// Add Kinematics Node
+							for (int i = Kinematics.NumKinematicBindings() - 1; i >= 0; i--)
+							{
+								FKinematics::FBindingKey Key = Kinematics.GetKinematicBindingKey(i);
+
+								int32 BoneIndex = INDEX_NONE;
+								TArray<int32> BoundVerts;
+								TArray<float> BoundWeights;
+								Kinematics.GetBoneBindings(Key, BoneIndex, BoundVerts, BoundWeights);
+
+								for (int32 vdx : BoundVerts)
+								{
+									ParticleIsKinematic[vdx] = true;
+								}
+							}
+
+
+							TArray<Chaos::TVector<int32, 3>> IndicesArray;
+							for (int32 i = 0; i < Indices->Num(); i++)
+							{
+								Chaos::TVector<int32, 3> CurrentIndices(0);
+								for (int32 j = 0; j < 3; j++)
+								{
+									CurrentIndices[j] = (*Indices)[i][j];
+								}
+								if (CurrentIndices[0] != -1
+									&& CurrentIndices[1] != -1
+									&& CurrentIndices[2] != -1)
+								{
+									IndicesArray.Emplace(CurrentIndices);
+								}
+							}
+							TArray<TArray<int32>> LocalIndex;
+							TArray<TArray<int32>>* LocalIndexPtr = &LocalIndex;
+							TArray<TArray<int>> GlobalIndex = Chaos::Utilities::ComputeIncidentElements(IndicesArray, LocalIndexPtr);
+							int32 ActualParticleCount = 0;
+							for (int32 l = 0; l < GlobalIndex.Num(); l++)
+							{
+								if (GlobalIndex[l].Num() > 0)
+								{
+									ActualParticleCount += 1;
+								}
+							}
+							TArray<Chaos::TVector<float, 3>> IndicesPositions;
+							IndicesPositions.SetNum(ActualParticleCount);
+							TArray<int32> IndicesMap;
+							IndicesMap.SetNum(ActualParticleCount);
+							int32 CurrentParticleIndex = 0;
+							for (int32 i = 0; i < GlobalIndex.Num(); i++)
+							{
+								if (GlobalIndex[i].Num() > 0)
+								{
+									IndicesPositions[CurrentParticleIndex] = (*Vertices)[(*Indices)[GlobalIndex[i][0]][LocalIndex[i][0]]];
+									IndicesMap[CurrentParticleIndex] = (*Indices)[GlobalIndex[i][0]][LocalIndex[i][0]];
+									CurrentParticleIndex += 1;
+								}
+							}
+
+
+							//for (int32 i = 0; i < Vertices->Num(); i++)
+							for (int32 i = 0; i < IndicesMap.Num(); i++)
+							{
+								//only work on particles that are not kinematic:
+								if (!ParticleIsKinematic[IndicesMap[i]])
+								{
+									TArray<int32> ParticleIntersection = VertexBVH.FindAllIntersections((*Vertices)[IndicesMap[i]]);
+									int32 MinIndex = -1;
+									float MinDis = SphereRadius;
+									for (int32 j = 0; j < ParticleIntersection.Num(); j++)
+									{
+										Chaos::FRealSingle CurrentDistance = ((*Vertices)[IndicesMap[i]] - PositionVertexBuffer.VertexPosition(ParticleIntersection[j])).Size();
+										if (CurrentDistance < MinDis)
+										{
+											MinDis = CurrentDistance;
+											MinIndex = ParticleIntersection[j];
+										}
+									}
+
+									if (MinIndex != -1)
+									{
+										int32 BoneParticleIndex = -1;
+
+										int32 SectionIndex;
+										int32 VertIndex;
+										LODRenderData->GetSectionFromVertexIndex(MinIndex, SectionIndex, VertIndex);
+
+										check(SectionIndex < LODRenderData->RenderSections.Num());
+										const FSkelMeshRenderSection& Section = LODRenderData->RenderSections[SectionIndex];
+										int32 BufferVertIndex = Section.GetVertexBufferIndex() + VertIndex;
+										for (int32 InfluenceIndex = 0; InfluenceIndex < MaxBoneInfluences; InfluenceIndex++)
+										{
+											const int32 BoneIndex = Section.BoneMap[SkinWeightVertexBuffer->GetBoneIndex(BufferVertIndex, InfluenceIndex)];
+											const float	Weight = (float)SkinWeightVertexBuffer->GetBoneWeight(BufferVertIndex, InfluenceIndex) * UE::AnimationCore::InvMaxRawBoneWeightFloat;
+											if (Weight > float(0) && 0 <= BoneIndex && BoneIndex < ComponentPose.Num())
+											{
+												BoneBoundVerts[BoneIndex].Add(IndicesMap[i]);
+												//BoneBoundWeights[BoneIndex].Add(Weight);
+												BoneBoundWeights[BoneIndex].Add(1.f);
+												break;
+											}
+										}
+									}
+								}
+							}
+						}
+						else
+						{
+							TMap<int32, int32> BoneVertexToCollectionMap;
+							GeometryCollection::Facades::FPositionTargetFacade PositionTargets(InCollection);
+							PositionTargets.DefineSchema();
+
+							for (int32 i = 0; i < Indices->Num(); i++)
+							{
+								TArray<int32> TriangleIntersections0 = VertexBVH.FindAllIntersections((*Vertices)[(*Indices)[i][0]]);
+								TArray<int32> TriangleIntersections1 = VertexBVH.FindAllIntersections((*Vertices)[(*Indices)[i][1]]);
+								TArray<int32> TriangleIntersections2 = VertexBVH.FindAllIntersections((*Vertices)[(*Indices)[i][2]]);
+								TriangleIntersections0.Sort();
+								TriangleIntersections1.Sort();
+								TriangleIntersections2.Sort();
+								TArray<int32> TriangleIntersections({});
+								for (int32 k = 0; k < TriangleIntersections0.Num(); k++)
+								{
+									if (TriangleIntersections1.Contains(TriangleIntersections0[k])
+										&& TriangleIntersections2.Contains(TriangleIntersections0[k]))
+									{
+										TriangleIntersections.Emplace(TriangleIntersections0[k]);
+									}
+								}
+								int32 MinIndex = -1;
+								float MinDis = SphereRadius;
+								Chaos::TVector<float, 3> ClosestBary(0.f);
+								for (int32 j = 0; j < TriangleIntersections.Num(); j++)
+								{
+									Chaos::TVector<float, 3> Bary,
+										TriPos0((*Vertices)[(*Indices)[i][0]]), TriPos1((*Vertices)[(*Indices)[i][1]]), TriPos2((*Vertices)[(*Indices)[i][2]]),
+										ParticlePos(PositionVertexBuffer.VertexPosition(TriangleIntersections[j]));
+									Chaos::TVector<Chaos::FRealSingle, 3> ClosestPoint = Chaos::FindClosestPointAndBaryOnTriangle(TriPos0, TriPos1, TriPos2, ParticlePos, Bary);
+									Chaos::FRealSingle CurrentDistance = (ParticlePos - ClosestPoint).Size();
+									if (CurrentDistance < MinDis)
+									{
+										MinDis = CurrentDistance;
+										MinIndex = TriangleIntersections[j];
+										ClosestBary = Bary;
+									}
+								}
+
+								if (MinIndex != -1)
+								{
+									int32 BoneParticleIndex = -1;
+									//TODO: add kinematic particles first
+									if (!BoneVertexToCollectionMap.Contains(MinIndex))
+									{
+
+										int32 ParticleIndex = InCollection.AddElements(1, FGeometryCollection::VerticesGroup);
+										TManagedArray<FVector3f>& CurrentVertices = InCollection.ModifyAttribute<FVector3f>("Vertex", FGeometryCollection::VerticesGroup);
+
+										CurrentVertices[ParticleIndex] = PositionVertexBuffer.VertexPosition(MinIndex);
+										BoneVertexToCollectionMap.Emplace(MinIndex, ParticleIndex);
+
+										int32 SectionIndex;
+										int32 VertIndex;
+										LODRenderData->GetSectionFromVertexIndex(MinIndex, SectionIndex, VertIndex);
+
+										check(SectionIndex < LODRenderData->RenderSections.Num());
+										const FSkelMeshRenderSection& Section = LODRenderData->RenderSections[SectionIndex];
+										int32 BufferVertIndex = Section.GetVertexBufferIndex() + VertIndex;
+										for (int32 InfluenceIndex = 0; InfluenceIndex < MaxBoneInfluences; InfluenceIndex++)
+										{
+											const int32 BoneIndex = Section.BoneMap[SkinWeightVertexBuffer->GetBoneIndex(BufferVertIndex, InfluenceIndex)];
+											const float	Weight = (float)SkinWeightVertexBuffer->GetBoneWeight(BufferVertIndex, InfluenceIndex) * UE::AnimationCore::InvMaxRawBoneWeightFloat;
+											if (Weight > float(0) && 0 <= BoneIndex && BoneIndex < ComponentPose.Num())
+											{
+												BoneBoundVerts[BoneIndex].Add(ParticleIndex);
+												BoneBoundWeights[BoneIndex].Add(Weight);
+											}
+										}
+									}
+
+									BoneParticleIndex = BoneVertexToCollectionMap[MinIndex];
+
+									GeometryCollection::Facades::FPositionTargetsData DataPackage;
+									DataPackage.TargetIndex.Init(BoneParticleIndex, 1);
+									DataPackage.TargetWeights.Init(1.f, 1);
+									DataPackage.SourceWeights.Init(1.f, 3);
+									DataPackage.SourceIndex.Init(-1, 3);
+									DataPackage.SourceIndex[0] = (*Indices)[i][0];
+									DataPackage.SourceIndex[1] = (*Indices)[i][1];
+									DataPackage.SourceIndex[2] = (*Indices)[i][2];
+									DataPackage.SourceWeights[0] = ClosestBary[0];
+									DataPackage.SourceWeights[1] = ClosestBary[1];
+									DataPackage.SourceWeights[2] = ClosestBary[2];
+									if (TManagedArray<float>* Mass = InCollection.FindAttribute<float>("Mass", FGeometryCollection::VerticesGroup))
+									{
+										DataPackage.Stiffness = 0.f;
+										for (int32 k = 0; k < 3; k++)
+										{
+											DataPackage.Stiffness += DataPackage.SourceWeights[k] * PositionTargetStiffness * (*Mass)[DataPackage.SourceIndex[k]];
+										}
+										//DataPackage.Stiffness += DataPackage.TargetWeights[0] * PositionTargetStiffness * (*Mass)[DataPackage.TargetIndex[0]];
+										//DataPackage.Stiffness /= 2.f;
+									}
+									else
+									{
+										DataPackage.Stiffness = PositionTargetStiffness;
+									}
+									PositionTargets.AddPositionTarget(DataPackage);
+								}
+
+
+							}
+
+
+						}
+						auto DoubleVert = [](FVector3f V) { return FVector3d(V.X, V.Y, V.Z); };
+
+						for (int32 BoneIndex = 0; BoneIndex < ComponentPose.Num(); ++BoneIndex)
+						{
+							if (BoneBoundVerts[BoneIndex].Num())
+							{
+								FString BoneName = BoneSkeletalMesh->GetRefSkeleton().GetBoneName(BoneIndex).ToString();
+								//get local coords of bound verts
+								typedef GeometryCollection::Facades::FKinematicBindingFacade FKinematics;
+								FKinematics Kinematics(InCollection); Kinematics.DefineSchema();
+								if (Kinematics.IsValid())
+								{
+									FKinematics::FBindingKey Binding = Kinematics.SetBoneBindings(BoneIndex, BoneBoundVerts[BoneIndex], BoneBoundWeights[BoneIndex]);
+									TManagedArray<TArray<FVector3f>>& LocalPos = InCollection.AddAttribute<TArray<FVector3f>>("LocalPosition", Binding.GroupName);
+									Kinematics.AddKinematicBinding(Binding);
+
+									auto FloatVert = [](FVector3d V) { return FVector3f(V.X, V.Y, V.Z); };
+									LocalPos[Binding.Index].SetNum(BoneBoundVerts[BoneIndex].Num());
+									for (int32 i = 0; i < BoneBoundVerts[BoneIndex].Num(); i++)
+									{
+										FVector3f Temp = (*Vertices)[BoneBoundVerts[BoneIndex][i]];
+										LocalPos[Binding.Index][i] = FloatVert(ComponentPose[BoneIndex].InverseTransformPosition(DoubleVert(Temp)));
+									}
+								}
+							}
+						}
+						GeometryCollection::Facades::FVertexBoneWeightsFacade(InCollection).AddBoneWeightsFromKinematicBindings();
+
+					}
+
+
+				}
+			}
+		}
+		SetValue<DataType>(Context, InCollection, &Collection);
+	}
 }
\ No newline at end of file
diff --git a/Engine/Plugins/Experimental/ChaosFlesh/Source/ChaosFleshNodes/Public/Dataflow/ChaosFleshPositionTargetInitializationNodes.h b/Engine/Plugins/Experimental/ChaosFlesh/Source/ChaosFleshNodes/Public/Dataflow/ChaosFleshPositionTargetInitializationNodes.h
--- a/Engine/Plugins/Experimental/ChaosFlesh/Source/ChaosFleshNodes/Public/Dataflow/ChaosFleshPositionTargetInitializationNodes.h
+++ b/Engine/Plugins/Experimental/ChaosFlesh/Source/ChaosFleshNodes/Public/Dataflow/ChaosFleshPositionTargetInitializationNodes.h
@@ -205,6 +205,53 @@ public:
 	virtual void Evaluate(Dataflow::FContext& Context, const FDataflowOutput* Out) const override;
 };
 
+UENUM(BlueprintType)
+enum class ESkeletalBindingMode : uint8
+{
+	Dataflow_SkeletalBinding_Kinematic UMETA(DisplayName = "Kinematic"),
+	Dataflow_SkeletalBinding_PositionTarget UMETA(DisplayName = "Position Target"),
+	//
+	Chaos_Max UMETA(Hidden)
+};
+
+USTRUCT(meta = (DataflowFlesh))
+struct FSetFleshBonePositionTargetBindingDataflowNode : public FDataflowNode
+{
+	GENERATED_USTRUCT_BODY()
+		DATAFLOW_NODE_DEFINE_INTERNAL(FSetFleshBonePositionTargetBindingDataflowNode, "FSetFleshBonePositionTargetBinding", "Flesh", "")
+		DATAFLOW_NODE_RENDER_TYPE(FGeometryCollection::StaticType(), "Collection")
+
+public:
+	typedef FManagedArrayCollection DataType;
+
+
+	UPROPERTY(meta = (DataflowInput, DataflowOutput, DisplayName = "Collection", DataflowPassthrough = "Collection"))
+		FManagedArrayCollection Collection;
+
+	UPROPERTY(EditAnywhere, Category = "Dataflow")
+		ESkeletalBindingMode SkeletalBindingMode = ESkeletalBindingMode::Dataflow_SkeletalBinding_PositionTarget;
+
+	UPROPERTY(EditAnywhere, Category = "Dataflow")
+		float PositionTargetStiffness = 10000.f;
+
+	UPROPERTY(meta = (DataflowInput, DisplayName = "SkeletalMesh"))
+		TObjectPtr<USkeletalMesh> SkeletalMeshIn = nullptr;
+
+	UPROPERTY(EditAnywhere, Category = "Dataflow", meta = (ClampMin = "0.0", ClampMax = "2.0", UIMin = "0.0", UIMax = "2.0"))
+		float VertexRadiusRatio = .001f;
+
+
+	FSetFleshBonePositionTargetBindingDataflowNode(const Dataflow::FNodeParameters& InParam, FGuid InGuid = FGuid::NewGuid())
+		: FDataflowNode(InParam, InGuid)
+	{
+		RegisterInputConnection(&SkeletalMeshIn);
+		RegisterInputConnection(&Collection);
+		RegisterOutputConnection(&Collection, &Collection);
+	}
+
+	virtual void Evaluate(Dataflow::FContext& Context, const FDataflowOutput* Out) const override;
+};
+
 namespace Dataflow
 {
 	void RegisterChaosFleshPositionTargetInitializationNodes();
diff --git a/Engine/Source/Runtime/Experimental/Chaos/Private/GeometryCollection/Facades/CollectionVertexBoneWeightsFacade.cpp b/Engine/Source/Runtime/Experimental/Chaos/Private/GeometryCollection/Facades/CollectionVertexBoneWeightsFacade.cpp
--- a/Engine/Source/Runtime/Experimental/Chaos/Private/GeometryCollection/Facades/CollectionVertexBoneWeightsFacade.cpp
+++ b/Engine/Source/Runtime/Experimental/Chaos/Private/GeometryCollection/Facades/CollectionVertexBoneWeightsFacade.cpp
@@ -91,6 +91,17 @@ namespace GeometryCollection::Facades
 			TManagedArray< TArray<int32> >& IndicesArray = BoneIndexAttribute.Modify();
 			TManagedArray< TArray<float> >& WeightsArray = BoneWeightAttribute.Modify();
 
+			TArray<float> TotalWeights;
+			TotalWeights.Init(0.f, WeightsArray.Num());
+
+			for (int32 Vert = 0; Vert < WeightsArray.Num(); Vert++)
+			{
+				for (int32 i = 0; i < WeightsArray[Vert].Num(); i++)
+				{
+					TotalWeights[Vert] += WeightsArray[Vert][i];
+				}
+			}
+
 			GeometryCollection::Facades::FKinematicBindingFacade BindingFacade(ConstCollection);
 			for (int32 Kdx = BindingFacade.NumKinematicBindings() - 1; 0 <= Kdx; Kdx--)
 			{
@@ -108,19 +119,29 @@ namespace GeometryCollection::Facades
 						if (0 <= Vert && Vert < NumVertices && !IndicesArray[Vert].Contains(Bone))
 						{
 							int32 BoneIndex = IndicesArray[Vert].Find(Bone);
-							if (BoneIndex == INDEX_NONE)
+							if (TotalWeights[Vert] + Weight <= 1.f)
 							{
-								IndicesArray[Vert].Add(Bone);
-								WeightsArray[Vert].Add(Weight);
+								if (BoneIndex == INDEX_NONE)
+								{
+									IndicesArray[Vert].Add(Bone);
+									WeightsArray[Vert].Add(Weight);
+									TotalWeights[Vert] += Weight;
+								}
+								else if (0 <= BoneIndex && BoneIndex < WeightsArray[Vert].Num())
+								{
+									WeightsArray[Vert][BoneIndex] = Weight;
+									TotalWeights[Vert] += Weight;
+								}
 							}
-							else if (0 <= BoneIndex && BoneIndex < WeightsArray[Vert].Num())
+							else
 							{
-								WeightsArray[Vert][BoneIndex] = Weight;
+								UE_LOG(LogChaos, Warning, TEXT("Bone weight sum exceeds 1 on vertex %d"), Vert);
 							}
 						}
 					}
 				}
 			}
+
 		}
 	}
 
