ChaosFlesh: Render muscle components

diff --git a/Engine/Plugins/Experimental/ChaosFlesh/Source/ChaosFleshNodes/Private/Dataflow/ChaosFleshCoreNodes.cpp b/Engine/Plugins/Experimental/ChaosFlesh/Source/ChaosFleshNodes/Private/Dataflow/ChaosFleshCoreNodes.cpp
--- a/Engine/Plugins/Experimental/ChaosFlesh/Source/ChaosFleshNodes/Private/Dataflow/ChaosFleshCoreNodes.cpp
+++ b/Engine/Plugins/Experimental/ChaosFlesh/Source/ChaosFleshNodes/Private/Dataflow/ChaosFleshCoreNodes.cpp
@@ -2,7 +2,8 @@
 
 #include "Dataflow/ChaosFleshCoreNodes.h"
 
-
+#include "ChaosFlesh/FleshCollection.h"
+#include "Dataflow/DataflowNodeFactory.h"
 
 
 namespace Dataflow
@@ -12,7 +13,33 @@ void RegisterChaosFleshCoreNodes()
 {
 	//DATAFLOW_NODE_REGISTER_CREATION_FACTORY(FNewManagedArrayCollectionNode);
 	//DATAFLOW_NODE_REGISTER_CREATION_FACTORY(FAddAttributeNode);
+	DATAFLOW_NODE_REGISTER_CREATION_FACTORY(FAppendTetrahedralCollectionDataflowNode);
+}
+
 }
 
+void FAppendTetrahedralCollectionDataflowNode::Evaluate(Dataflow::FContext& Context, const FDataflowOutput* Out) const
+{
+	if (Out->IsA<FManagedArrayCollection>(&Collection1))
+	{
+		TUniquePtr<FFleshCollection> InCollection1(GetValue<FManagedArrayCollection>(Context, &Collection1).NewCopy<FFleshCollection>());
+		TUniquePtr<FFleshCollection> InCollection2(GetValue<FManagedArrayCollection>(Context, &Collection2).NewCopy<FFleshCollection>());
+		TArray<FString> GeometryGroupGuidsLocal1, GeometryGroupGuidsLocal2;
+		if (const TManagedArray<FString>* GuidArray1 = InCollection1->FindAttribute<FString>("Guid", FGeometryCollection::GeometryGroup))
+		{
+			GeometryGroupGuidsLocal1 = GuidArray1->GetConstArray();
+		}
+		if (InCollection2)
+		{
+			InCollection1->AppendGeometry(*InCollection2);
+		}		
+		if (const TManagedArray<FString>* GuidArray2 = InCollection2->FindAttribute<FString>("Guid", FGeometryCollection::GeometryGroup))
+		{
+			GeometryGroupGuidsLocal2 = GuidArray2->GetConstArray();
+		}
+		SetValue<FManagedArrayCollection>(Context, *InCollection1, &Collection1);
+		SetValue<TArray<FString>>(Context, GeometryGroupGuidsLocal1, &GeometryGroupGuidsOut1);
+		SetValue<TArray<FString>>(Context, GeometryGroupGuidsLocal2, &GeometryGroupGuidsOut2);
+	}
 }
 
diff --git a/Engine/Plugins/Experimental/ChaosFlesh/Source/ChaosFleshNodes/Private/Dataflow/ChaosFleshEngineAssetNodes.cpp b/Engine/Plugins/Experimental/ChaosFlesh/Source/ChaosFleshNodes/Private/Dataflow/ChaosFleshEngineAssetNodes.cpp
--- a/Engine/Plugins/Experimental/ChaosFlesh/Source/ChaosFleshNodes/Private/Dataflow/ChaosFleshEngineAssetNodes.cpp
+++ b/Engine/Plugins/Experimental/ChaosFlesh/Source/ChaosFleshNodes/Private/Dataflow/ChaosFleshEngineAssetNodes.cpp
@@ -1,6 +1,7 @@
 // Copyright Epic Games, Inc. All Rights Reserved.
 
 #include "Dataflow/ChaosFleshEngineAssetNodes.h"
+#include "Dataflow/ChaosFleshTetrahedralNodes.h"
 
 #include "Chaos/Math/Poisson.h"
 #include "Chaos/Utilities.h"
@@ -25,6 +26,8 @@ namespace Dataflow
 		DATAFLOW_NODE_REGISTER_CREATION_FACTORY(FComputeFiberFieldNode);
 		DATAFLOW_NODE_REGISTER_CREATION_FACTORY(FComputeIslandsNode);
 		DATAFLOW_NODE_REGISTER_CREATION_FACTORY(FGenerateOriginInsertionNode);
+		DATAFLOW_NODE_REGISTER_CREATION_FACTORY(FIsolateComponentNode);
+		DATAFLOW_NODE_REGISTER_CREATION_FACTORY(FGetSurfaceIndicesNode);
 	}
 }
 
@@ -554,4 +557,91 @@ void FGenerateOriginInsertionNode::Evaluate(Dataflow::FContext& Context, const F
 		SetValue<TArray<int32>>(Context, OutOriginIndices, &OriginIndicesOut);
 		SetValue<TArray<int32>>(Context, OutInsertionIndices, &InsertionIndicesOut);
 	}
+}
+
+void FIsolateComponentNode::Evaluate(Dataflow::FContext& Context, const FDataflowOutput* Out) const
+{
+	if (Out->IsA<FManagedArrayCollection>(&Collection))
+	{	
+		FManagedArrayCollection InCollection = GetValue<FManagedArrayCollection>(Context, &Collection);
+		TArray<int32> DeleteList;
+		if (TManagedArray<FIntVector>* Indices = InCollection.FindAttribute<FIntVector>("Indices", FGeometryCollection::FacesGroup))
+		{
+			if (TManagedArray<bool>* FaceVisibility = InCollection.FindAttribute<bool>("Visible", FGeometryCollection::FacesGroup))
+			{
+				if (TManagedArray<int32>* ComponentIndex = InCollection.FindAttribute<int32>("ComponentIndex", FGeometryCollection::VerticesGroup))
+				{	
+					TSet<int32> ComponentSet;
+					TArray<FString> StrArray;
+					TargetComponentIndex.ParseIntoArray(StrArray, *FString(" "));
+					for (FString Elem : StrArray)
+					{
+						if (Elem.Len() && FCString::IsNumeric(*Elem))
+						{
+							ComponentSet.Add(FCString::Atoi(*Elem));
+						}
+					}
+					for (int32 i = 0; i < Indices->Num(); i++)
+					{
+						if (!ComponentSet.Contains((*ComponentIndex)[(*Indices)[i][0]]))
+						{
+							(*FaceVisibility)[i] = false;
+							DeleteList.AddUnique(i);
+						}
+					}
+				}
+			}
+		}
+		if (bDeleteHiddenFaces)
+		{
+			DeleteList.Sort();
+			InCollection.RemoveElements(FGeometryCollection::FacesGroup, DeleteList);
+		}
+		SetValue<FManagedArrayCollection>(Context, InCollection, &Collection);
+	}
+}
+
+void FGetSurfaceIndicesNode::Evaluate(Dataflow::FContext& Context, const FDataflowOutput* Out) const
+{
+	FManagedArrayCollection InCollection = GetValue<FManagedArrayCollection>(Context, &Collection);
+	TArray<int32> SurfaceIndicesLocal;
+	if (TManagedArray<FIntVector>* Indices = InCollection.FindAttribute<FIntVector>("Indices", FGeometryCollection::FacesGroup))
+	{
+		if (FindInput(&GeometryGroupGuidsIn) && FindInput(&GeometryGroupGuidsIn)->GetConnection())
+		{
+			TArray<FString> GeometryGroupGuidsLocal = GetValue<TArray<FString>>(Context, &GeometryGroupGuidsIn);
+			TManagedArray<int32>* IndicesStart = InCollection.FindAttribute<int32>("FaceStart", FGeometryCollection::GeometryGroup);
+			TManagedArray<int32>* IndicesCount = InCollection.FindAttribute<int32>("FaceCount", FGeometryCollection::GeometryGroup);
+			if (TManagedArray<FString>* Guids = InCollection.FindAttribute<FString>("Guid", FGeometryCollection::GeometryGroup))
+			{
+				for (int32 Idx = 0; Idx < IndicesStart->Num(); Idx++)
+				{
+					if (GeometryGroupGuidsLocal.Num() && Guids)
+					{
+						if (GeometryGroupGuidsLocal.Contains((*Guids)[Idx]))
+						{
+							for (int32 i = (*IndicesStart)[Idx]; i < (*IndicesStart)[Idx] + (*IndicesCount)[Idx]; i++)
+							{
+								for (int32 j = 0; j < 3; j++)
+								{
+									SurfaceIndicesLocal.AddUnique((*Indices)[i][j]);
+								}
+							}
+						}
+					}
+				}
+			}
+		}
+		else
+		{
+			for (int32 i = 0; i < Indices->Num(); i++)
+			{
+				for (int32 j = 0; j < 3; j++)
+				{
+					SurfaceIndicesLocal.AddUnique((*Indices)[i][j]);
+				}
+			}
+		}
+	}
+	SetValue<TArray<int32>>(Context, SurfaceIndicesLocal, &SurfaceIndicesOut);
 }
\ No newline at end of file
diff --git a/Engine/Plugins/Experimental/ChaosFlesh/Source/ChaosFleshNodes/Private/Dataflow/ChaosFleshTetrahedralNodes.cpp b/Engine/Plugins/Experimental/ChaosFlesh/Source/ChaosFleshNodes/Private/Dataflow/ChaosFleshTetrahedralNodes.cpp
--- a/Engine/Plugins/Experimental/ChaosFlesh/Source/ChaosFleshNodes/Private/Dataflow/ChaosFleshTetrahedralNodes.cpp
+++ b/Engine/Plugins/Experimental/ChaosFlesh/Source/ChaosFleshNodes/Private/Dataflow/ChaosFleshTetrahedralNodes.cpp
@@ -3,6 +3,7 @@
 #include "Dataflow/ChaosFleshTetrahedralNodes.h"
 #include "Dataflow/ChaosFleshEngineAssetNodes.h"
 
+#include "Async/ParallelFor.h"
 #include "Chaos/Deformable/Utilities.h"
 #include "ChaosFlesh/ChaosFlesh.h"
 #include "Chaos/Tetrahedron.h"
@@ -255,17 +256,72 @@ void FGenerateTetrahedralCollectionDataflowNodes::Evaluate(Dataflow::FContext& C
 				Converter.Convert(&SourceMesh, DynamicMesh);
 			}
 
-			if (Method == TetMeshingMethod::IsoStuffing)
+			if (!bComputeByComponent)
 			{
-				EvaluateIsoStuffing(Context, InCollection, DynamicMesh);
-			}
-			else if (Method == TetMeshingMethod::TetWild)
-			{
-				EvaluateTetWild(Context, InCollection, DynamicMesh);
+				if (Method == TetMeshingMethod::IsoStuffing)
+				{
+					EvaluateIsoStuffing(Context, InCollection, DynamicMesh);
+				}
+				else if (Method == TetMeshingMethod::TetWild)
+				{
+					EvaluateTetWild(Context, InCollection, DynamicMesh);
+				}
+				else
+				{
+					ensureMsgf(false, TEXT("FGenerateTetrahedralCollectionDataflowNodes unsupported Method."));
+				}
 			}
 			else
 			{
-				ensureMsgf(false, TEXT("FGenerateTetrahedralCollectionDataflowNodes unsupported Method."));
+				TArray<TArray<int32>> ConnectedComponents;
+				
+				TArray<FIntVector3> Faces;
+				Faces.SetNum(DynamicMesh.TriangleCount());
+				for (int32 i = 0; i < DynamicMesh.TriangleCount(); ++i)
+				{
+					Faces[i] = FIntVector3(DynamicMesh.GetTriangle(i));
+				}
+				Chaos::Utilities::FindConnectedRegions(Faces, ConnectedComponents);
+				TArray<TUniquePtr<FFleshCollection>> CollectionBuffer;
+				for (int32 i = 0; i < ConnectedComponents.Num(); i++)
+				{
+					CollectionBuffer.Add(TUniquePtr<FFleshCollection>(new FFleshCollection()));
+					//CollectionBuffer[i]->AddElement(1, FGeometryCollection::TransformGroup);
+				}
+				ParallelFor(ConnectedComponents.Num(), [&](int32 i)
+				{		
+					UE::Geometry::FDynamicMesh3 ComponentDynamicMesh;
+					TArray<FIntVector3> ComponentFaces;
+					ComponentFaces.SetNum(ConnectedComponents[i].Num());
+					for (FVector V : DynamicMesh.VerticesItr())
+					{
+						ComponentDynamicMesh.AppendVertex(V);
+					}
+					for (int32 j = 0; j < ConnectedComponents[i].Num(); j++)
+					{
+						int32 ElementIndex = ConnectedComponents[i][j];
+						for (int32 ie = 0; ie < 3; ie++)
+						{
+							ComponentDynamicMesh.AppendTriangle(Faces[ElementIndex][0], Faces[ElementIndex][1], Faces[ElementIndex][2]);
+						}
+					}
+					ComponentDynamicMesh.CompactInPlace();
+					if (Method == TetMeshingMethod::IsoStuffing)
+					{
+						EvaluateIsoStuffing(Context, CollectionBuffer[i], ComponentDynamicMesh);
+					}
+					else if (Method == TetMeshingMethod::TetWild)
+					{
+						EvaluateTetWild(Context, CollectionBuffer[i], ComponentDynamicMesh);
+					}
+				});
+				for (int32 i = 0; i < ConnectedComponents.Num(); i++)
+				{
+					if (CollectionBuffer[i]->NumElements(FGeometryCollection::VerticesGroup))
+					{
+						int32 GeomIndex = InCollection->AppendGeometry(*CollectionBuffer[i].Get());
+					}
+				}
 			}
 #else
 			ensureMsgf(false, TEXT("FGenerateTetrahedralCollectionDataflowNodes is an editor only node."));
diff --git a/Engine/Plugins/Experimental/ChaosFlesh/Source/ChaosFleshNodes/Public/Dataflow/ChaosFleshCoreNodes.h b/Engine/Plugins/Experimental/ChaosFlesh/Source/ChaosFleshNodes/Public/Dataflow/ChaosFleshCoreNodes.h
--- a/Engine/Plugins/Experimental/ChaosFlesh/Source/ChaosFleshNodes/Public/Dataflow/ChaosFleshCoreNodes.h
+++ b/Engine/Plugins/Experimental/ChaosFlesh/Source/ChaosFleshNodes/Public/Dataflow/ChaosFleshCoreNodes.h
@@ -8,6 +8,41 @@
 #include "Dataflow/DataflowNode.h"
 #include "GeometryCollection/ManagedArray.h"
 #include "GeometryCollection/ManagedArrayCollection.h"
+#include "GeometryCollection/GeometryCollection.h"
+
+#include "ChaosFleshCoreNodes.generated.h"
+
+USTRUCT(meta = (DataflowFlesh))
+struct FAppendTetrahedralCollectionDataflowNode : public FDataflowNode
+{
+	GENERATED_USTRUCT_BODY()
+	DATAFLOW_NODE_DEFINE_INTERNAL(FAppendTetrahedralCollectionDataflowNode, "AppendTetrahedralCollection", "Flesh", "")
+	DATAFLOW_NODE_RENDER_TYPE(FGeometryCollection::StaticType(), "Collection")
+public:
+	UPROPERTY(meta = (DataflowInput, DataflowOutput, DisplayName = "Collection1", DataflowPassthrough = "Collection1"));
+	FManagedArrayCollection Collection1;
+
+	UPROPERTY(meta = (DataflowInput, DisplayName = "Collection2"));
+	FManagedArrayCollection Collection2;
+
+	UPROPERTY(meta = (DataflowOutput, DisplayName = "GeometryGroupIndicesOut1"))
+		TArray<FString> GeometryGroupGuidsOut1;
+
+	UPROPERTY(meta = (DataflowOutput, DisplayName = "GeometryGroupIndicesOut2"))
+		TArray<FString> GeometryGroupGuidsOut2;
+
+	FAppendTetrahedralCollectionDataflowNode(const Dataflow::FNodeParameters& InParam, FGuid InGuid = FGuid::NewGuid())
+		: FDataflowNode(InParam, InGuid)
+	{
+		RegisterInputConnection(&Collection1);
+		RegisterInputConnection(&Collection2);
+		RegisterOutputConnection(&Collection1, &Collection1);
+		RegisterOutputConnection(&GeometryGroupGuidsOut1);
+		RegisterOutputConnection(&GeometryGroupGuidsOut2);
+	}
+
+	virtual void Evaluate(Dataflow::FContext& Context, const FDataflowOutput* Out) const override;
+};
 
 namespace Dataflow
 {
@@ -122,6 +157,8 @@ namespace Dataflow
 
 
 	*/
+
+
 	void RegisterChaosFleshCoreNodes();
 }
 
diff --git a/Engine/Plugins/Experimental/ChaosFlesh/Source/ChaosFleshNodes/Public/Dataflow/ChaosFleshEngineAssetNodes.h b/Engine/Plugins/Experimental/ChaosFlesh/Source/ChaosFleshNodes/Public/Dataflow/ChaosFleshEngineAssetNodes.h
--- a/Engine/Plugins/Experimental/ChaosFlesh/Source/ChaosFleshNodes/Public/Dataflow/ChaosFleshEngineAssetNodes.h
+++ b/Engine/Plugins/Experimental/ChaosFlesh/Source/ChaosFleshNodes/Public/Dataflow/ChaosFleshEngineAssetNodes.h
@@ -231,6 +231,63 @@ public:
 	virtual void Evaluate(Dataflow::FContext& Context, const FDataflowOutput* Out) const override;
 };
 
+USTRUCT(meta = (DataflowFlesh))
+struct FIsolateComponentNode : public FDataflowNode
+{
+	GENERATED_USTRUCT_BODY()
+	DATAFLOW_NODE_DEFINE_INTERNAL(FIsolateComponentNode, "IsolateComponent", "Flesh", "")
+	DATAFLOW_NODE_RENDER_TYPE(FGeometryCollection::StaticType(), "Collection")
+
+public:
+	//typedef FManagedArrayCollection DataType;
+
+	UPROPERTY(meta = (DataflowInput, DataflowOutput, DisplayName = "Collection", DataflowPassthrough = "Collection"))
+		FManagedArrayCollection Collection;
+
+	UPROPERTY(EditAnywhere, Category = "Dataflow")
+		bool bDeleteHiddenFaces = false;
+
+	UPROPERTY(EditAnywhere, Category = "Dataflow")
+		FString TargetComponentIndex = "";
+
+	FIsolateComponentNode(const Dataflow::FNodeParameters& InParam, FGuid InGuid = FGuid::NewGuid())
+		: FDataflowNode(InParam, InGuid)
+	{
+		RegisterInputConnection(&Collection);
+		RegisterOutputConnection(&Collection, &Collection);
+	}
+
+	virtual void Evaluate(Dataflow::FContext& Context, const FDataflowOutput* Out) const override;
+};
+
+USTRUCT(meta = (DataflowFlesh))
+struct FGetSurfaceIndicesNode : public FDataflowNode
+{
+	GENERATED_USTRUCT_BODY()
+	DATAFLOW_NODE_DEFINE_INTERNAL(FGetSurfaceIndicesNode, "GetSurfaceIndices", "Flesh", "")
+	DATAFLOW_NODE_RENDER_TYPE(FGeometryCollection::StaticType(), "Collection")
+
+public:
+	UPROPERTY(meta = (DataflowInput, DisplayName = "Collection"))
+		FManagedArrayCollection Collection;
+
+	UPROPERTY(meta = (DataflowInput, DisplayName = "GeometryGroupGuidsIn"))
+		TArray<FString> GeometryGroupGuidsIn;
+
+	UPROPERTY(meta = (DataflowOutput, DisplayName = "SurfaceIndicesOut"))
+		TArray<int32> SurfaceIndicesOut;
+
+	FGetSurfaceIndicesNode(const Dataflow::FNodeParameters& InParam, FGuid InGuid = FGuid::NewGuid())
+		: FDataflowNode(InParam, InGuid)
+	{
+		RegisterInputConnection(&Collection);
+		RegisterInputConnection(&GeometryGroupGuidsIn);
+		RegisterOutputConnection(&SurfaceIndicesOut);
+	}
+
+	virtual void Evaluate(Dataflow::FContext& Context, const FDataflowOutput* Out) const override;
+};
+
 namespace Dataflow
 {
 	void RegisterChaosFleshEngineAssetNodes();
diff --git a/Engine/Plugins/Experimental/ChaosFlesh/Source/ChaosFleshNodes/Public/Dataflow/ChaosFleshTetrahedralNodes.h b/Engine/Plugins/Experimental/ChaosFlesh/Source/ChaosFleshNodes/Public/Dataflow/ChaosFleshTetrahedralNodes.h
--- a/Engine/Plugins/Experimental/ChaosFlesh/Source/ChaosFleshNodes/Public/Dataflow/ChaosFleshTetrahedralNodes.h
+++ b/Engine/Plugins/Experimental/ChaosFlesh/Source/ChaosFleshNodes/Public/Dataflow/ChaosFleshTetrahedralNodes.h
@@ -150,6 +150,9 @@ public:
 	UPROPERTY(EditAnywhere, Category = "Dataflow", meta = (DataflowInput, DisplayName = "SkeletalMesh"))
 	TObjectPtr<const USkeletalMesh> SkeletalMesh = nullptr;
 
+	UPROPERTY(EditAnywhere, Category = "Dataflow")
+	bool bComputeByComponent = false;
+
 	UPROPERTY(meta = (DataflowInput, DataflowOutput, DisplayName = "Collection"))
 	FManagedArrayCollection Collection;
 
diff --git a/Engine/Plugins/Experimental/GeometryCollectionPlugin/Source/GeometryCollectionNodes/Private/Dataflow/GeometryCollectionNodes.cpp b/Engine/Plugins/Experimental/GeometryCollectionPlugin/Source/GeometryCollectionNodes/Private/Dataflow/GeometryCollectionNodes.cpp
--- a/Engine/Plugins/Experimental/GeometryCollectionPlugin/Source/GeometryCollectionNodes/Private/Dataflow/GeometryCollectionNodes.cpp
+++ b/Engine/Plugins/Experimental/GeometryCollectionPlugin/Source/GeometryCollectionNodes/Private/Dataflow/GeometryCollectionNodes.cpp
@@ -84,7 +84,7 @@ namespace Dataflow
 		DATAFLOW_NODE_REGISTER_CREATION_FACTORY(FMultiplyTransformDataflowNode);
 		DATAFLOW_NODE_REGISTER_CREATION_FACTORY(FInvertTransformDataflowNode);
 		DATAFLOW_NODE_REGISTER_CREATION_FACTORY(FSelectionToVertexListDataflowNode);
-		
+		DATAFLOW_NODE_REGISTER_CREATION_FACTORY(FUnionIntArraysDataflowNode);
 
 		// GeometryCollection
 		DATAFLOW_NODE_REGISTER_CREATION_FACTORY_NODE_COLORS_BY_CATEGORY("GeometryCollection", FLinearColor(0.55f, 0.45f, 1.0f), CDefaultNodeBodyTintColor);
@@ -130,9 +130,19 @@ void FAppendCollectionAssetsDataflowNode::Evaluate(Dataflow::FContext& Context,
 	{
 		FManagedArrayCollection InCollection1 = GetValue<DataType>(Context, &Collection1);
 		const FManagedArrayCollection& InCollection2 = GetValue<DataType>(Context, &Collection2);
-
+		TArray<FString> GeometryGroupGuidsLocal1, GeometryGroupGuidsLocal2;
+		if (const TManagedArray<FString>* GuidArray1 = InCollection1.FindAttribute<FString>("Guid", FGeometryCollection::GeometryGroup))
+		{
+			GeometryGroupGuidsLocal1 = GuidArray1->GetConstArray();
+		}
 		InCollection1.Append(InCollection2);
 		SetValue<DataType>(Context, InCollection1, &Collection1);
+		if (const TManagedArray<FString>* GuidArray2 = InCollection2.FindAttribute<FString>("Guid", FGeometryCollection::GeometryGroup))
+		{
+			GeometryGroupGuidsLocal2 = GuidArray2->GetConstArray();
+		}
+		SetValue<TArray<FString>>(Context, GeometryGroupGuidsLocal1, &GeometryGroupGuidsOut1);
+		SetValue<TArray<FString>>(Context, GeometryGroupGuidsLocal2, &GeometryGroupGuidsOut2);
 	}
 }
 
@@ -1365,3 +1375,22 @@ void FInvertTransformDataflowNode::Evaluate(Dataflow::FContext& Context, const F
 		SetValue<FTransform>(Context, OutXf, &OutTransform);
 	}
 }
+
+void FUnionIntArraysDataflowNode::Evaluate(Dataflow::FContext& Context, const FDataflowOutput* Out) const
+{
+	if (Out->IsA<TArray<int32>>(&OutArray))
+	{
+		TArray<int32> Array1 = GetValue<TArray<int32>>(Context, &InArray1);
+		TArray<int32> Array2 = GetValue<TArray<int32>>(Context, &InArray2);
+		TArray<int32> OutputArray;
+		for (int32 i = 0; i < Array1.Num(); i++)
+		{
+			OutputArray.AddUnique(Array1[i]);
+		}
+		for (int32 i = 0; i < Array2.Num(); i++)
+		{
+			OutputArray.AddUnique(Array2[i]);
+		}
+		SetValue<TArray<int32>>(Context, OutputArray, &OutArray);
+	}
+}
diff --git a/Engine/Plugins/Experimental/GeometryCollectionPlugin/Source/GeometryCollectionNodes/Public/Dataflow/GeometryCollectionNodes.h b/Engine/Plugins/Experimental/GeometryCollectionPlugin/Source/GeometryCollectionNodes/Public/Dataflow/GeometryCollectionNodes.h
--- a/Engine/Plugins/Experimental/GeometryCollectionPlugin/Source/GeometryCollectionNodes/Public/Dataflow/GeometryCollectionNodes.h
+++ b/Engine/Plugins/Experimental/GeometryCollectionPlugin/Source/GeometryCollectionNodes/Public/Dataflow/GeometryCollectionNodes.h
@@ -70,6 +70,11 @@ public:
 	UPROPERTY(meta = (DataflowInput, DisplayName = "Collection"))
 	FManagedArrayCollection Collection2;
 
+	UPROPERTY(meta = (DataflowOutput, DisplayName = "GeometryGroupIndicesOut1"))
+	TArray<FString> GeometryGroupGuidsOut1;
+
+	UPROPERTY(meta = (DataflowOutput, DisplayName = "GeometryGroupIndicesOut2"))
+		TArray<FString> GeometryGroupGuidsOut2;
 
 	FAppendCollectionAssetsDataflowNode(const Dataflow::FNodeParameters& InParam, FGuid InGuid = FGuid::NewGuid())
 		: FDataflowNode(InParam, InGuid)
@@ -77,6 +82,8 @@ public:
 		RegisterInputConnection(&Collection1);
 		RegisterInputConnection(&Collection2);
 		RegisterOutputConnection(&Collection1, &Collection1);
+		RegisterOutputConnection(&GeometryGroupGuidsOut1);
+		RegisterOutputConnection(&GeometryGroupGuidsOut2);
 	}
 
 	virtual void Evaluate(Dataflow::FContext& Context, const FDataflowOutput* Out) const override;
@@ -1675,6 +1682,32 @@ public:
 	virtual void Evaluate(Dataflow::FContext& Context, const FDataflowOutput* Out) const override;
 };
 
+USTRUCT()
+struct FUnionIntArraysDataflowNode : public FDataflowNode
+{
+	GENERATED_USTRUCT_BODY()
+		DATAFLOW_NODE_DEFINE_INTERNAL(FUnionIntArraysDataflowNode, "UnionIntArrays", "Utilities", "")
+
+public:
+	UPROPERTY(meta = (DataflowInput, DisplayName = "InArray1"));
+	TArray<int32> InArray1;
+
+	UPROPERTY(meta = (DataflowInput, DisplayName = "InArray2"));
+	TArray<int32> InArray2;
+
+	UPROPERTY(meta = (DataflowOutput, DisplayName = "OutArray", DataflowPassthrough = "InArray1"));
+	TArray<int32> OutArray;
+
+	FUnionIntArraysDataflowNode(const Dataflow::FNodeParameters& InParam, FGuid InGuid = FGuid::NewGuid())
+		: FDataflowNode(InParam, InGuid)
+	{
+		RegisterInputConnection(&InArray1);
+		RegisterInputConnection(&InArray2);
+		RegisterOutputConnection(&OutArray, &InArray1);
+	}
+
+	virtual void Evaluate(Dataflow::FContext& Context, const FDataflowOutput* Out) const override;
+};
 
 namespace Dataflow
 {
diff --git a/Engine/Source/Runtime/Experimental/Chaos/Public/Chaos/Utilities.h b/Engine/Source/Runtime/Experimental/Chaos/Public/Chaos/Utilities.h
--- a/Engine/Source/Runtime/Experimental/Chaos/Public/Chaos/Utilities.h
+++ b/Engine/Source/Runtime/Experimental/Chaos/Public/Chaos/Utilities.h
@@ -996,5 +996,59 @@ namespace Chaos
 			ConnectedComponentsDFSIterative(AdjacencyList, ConnectedComponents);
 		}
 
+		inline void FindConnectedRegions(const TArray<FIntVector3>& Elements, TArray<TArray<int32>>& ConnectedComponents)
+		{
+			TArray<int32> AllEntries, ElementIndices, Ordering, Ranges;
+			int32 count = 0;
+			AllEntries.Init(-1, Elements.Num() * 3);
+			ElementIndices.Init(-1, Elements.Num() * 3);
+			Ordering.Init(-1, Elements.Num() * 3);
+			Ranges.Init(-1, Elements.Num() * 3 + 1);
+			for (int32 i = 0; i < Elements.Num(); i++)
+			{
+				for (int32 j = 0; j < 3; j++)
+				{
+					AllEntries[3 * i + j] = Elements[i][j];
+					ElementIndices[3 * i + j] = i;
+					Ordering[3 * i + j] = count;
+					Ranges[3 * i + j] = count++;
+				}
+			}
+			Ranges[Elements.Num() * 3] = count;
+
+			if (AllEntries.Num() == 0)
+				return;
+
+			Ordering.Sort([&AllEntries](int32 a, int32 b) { return AllEntries[a] < AllEntries[b]; });
+			TArray<int32> UniqueRanges({ 0 });
+
+			for (int32 i = 0; i < Ranges.Num() - 2; ++i)
+			{
+				if (AllEntries[Ordering[i]] != AllEntries[Ordering[i + 1]])
+				{
+					UniqueRanges.Emplace(i + 1);
+				}
+			}
+
+			UniqueRanges.Emplace(count);
+
+			TArray<TArray<int32>> AdjacencyList;
+			AdjacencyList.Init(TArray<int32>(), Elements.Num());
+			for (int32 r = 0; r < UniqueRanges.Num() - 1; r++)
+			{
+				if (UniqueRanges[r + 1] - UniqueRanges[r] > 1)
+				{
+					for (int32 i = UniqueRanges[r]; i < UniqueRanges[r + 1]; i++)
+					{
+						for (int32 j = i + 1; j < UniqueRanges[r + 1]; j++)
+						{
+							AdjacencyList[ElementIndices[Ordering[i]]].Emplace(ElementIndices[Ordering[j]]);
+							AdjacencyList[ElementIndices[Ordering[j]]].Emplace(ElementIndices[Ordering[i]]);
+						}
+					}
+				}
+			}
+			ConnectedComponentsDFSIterative(AdjacencyList, ConnectedComponents);
+		}
 	} // namespace Utilities
 } // namespace Chaos
\ No newline at end of file
