ChaosFlesh: Add support for simultaneous muscle/flesh sim

diff --git a/Engine/Plugins/Experimental/ChaosFlesh/Source/ChaosFlesh/Private/ChaosFlesh/TetrahedralCollection.cpp b/Engine/Plugins/Experimental/ChaosFlesh/Source/ChaosFlesh/Private/ChaosFlesh/TetrahedralCollection.cpp
--- a/Engine/Plugins/Experimental/ChaosFlesh/Source/ChaosFlesh/Private/ChaosFlesh/TetrahedralCollection.cpp
+++ b/Engine/Plugins/Experimental/ChaosFlesh/Source/ChaosFlesh/Private/ChaosFlesh/TetrahedralCollection.cpp
@@ -14,6 +14,7 @@ const FName FTetrahedralCollection::TetrahedronStartAttribute("TetrahedronStart"
 const FName FTetrahedralCollection::TetrahedronCountAttribute("TetrahedronCount");
 const FName FTetrahedralCollection::IncidentElementsAttribute("IncidentElements");
 const FName FTetrahedralCollection::IncidentElementsLocalIndexAttribute("IncidentElementsLocalIndex");
+const FName FTetrahedralCollection::GuidAttribute("Guid");
 
 FTetrahedralCollection::FTetrahedralCollection()
 	: Super::FGeometryCollection()
@@ -37,6 +38,8 @@ void FTetrahedralCollection::Construct()
 	// Geometry Group
 	AddExternalAttribute<int32>(FTetrahedralCollection::TetrahedronStartAttribute, FGeometryCollection::GeometryGroup, TetrahedronStart, TetrahedronDependency);
 	AddExternalAttribute<int32>(FTetrahedralCollection::TetrahedronCountAttribute, FGeometryCollection::GeometryGroup, TetrahedronCount, TetrahedronDependency);
+	AddExternalAttribute<FString>(FTetrahedralCollection::GuidAttribute, FGeometryCollection::GeometryGroup, Guid);
+	for (FString& g : Guid) { g = FGuid::NewGuid().ToString(); }
 }
 
 
@@ -61,6 +64,7 @@ FTetrahedralCollection* FTetrahedralCollection::NewTetrahedralCollection(
 {
 	FTetrahedralCollection* Collection = new FTetrahedralCollection();
 	FTetrahedralCollection::Init(Collection, Vertices, SurfaceElements, Elements, bReverseVertexOrder);
+	for (FString& g : Collection->Guid) { g = FGuid::NewGuid().ToString(); }
 	return Collection;
 }
 void FTetrahedralCollection::Init(
@@ -161,6 +165,12 @@ int32 FTetrahedralCollection::AppendGeometry(
 		TetrahedronCount[NumGeometry + Idx] = Other.TetrahedronCount[Idx];
 	}
 
+	check(Guid.Num() == NumGeometry + Other.Guid.Num());
+	for (int32 Idx = 0; Idx < Other.Guid.Num(); Idx++)
+	{
+		Guid[NumGeometry + Idx] = Other.Guid[Idx];
+	}
+
 	// --- VERTICES GROUP --
 	int32 NumVertices = VerticesIndex;
 	int32 OtherNumVertices = Other.IncidentElements.Num();
diff --git a/Engine/Plugins/Experimental/ChaosFlesh/Source/ChaosFlesh/Public/ChaosFlesh/TetrahedralCollection.h b/Engine/Plugins/Experimental/ChaosFlesh/Source/ChaosFlesh/Public/ChaosFlesh/TetrahedralCollection.h
--- a/Engine/Plugins/Experimental/ChaosFlesh/Source/ChaosFlesh/Public/ChaosFlesh/TetrahedralCollection.h
+++ b/Engine/Plugins/Experimental/ChaosFlesh/Source/ChaosFlesh/Public/ChaosFlesh/TetrahedralCollection.h
@@ -87,6 +87,12 @@ public:
 	static const FName IncidentElementsLocalIndexAttribute;
 	TManagedArray<TArray<int32>> IncidentElementsLocalIndex;
 
+	/*
+	* Guid Attribute
+	*/
+	static const FName GuidAttribute;
+	TManagedArray<FString> Guid;
+
 protected:
 
 	void Construct();
diff --git a/Engine/Plugins/Experimental/ChaosFlesh/Source/ChaosFleshNodes/Private/Dataflow/ChaosFleshBindingsNodes.cpp b/Engine/Plugins/Experimental/ChaosFlesh/Source/ChaosFleshNodes/Private/Dataflow/ChaosFleshBindingsNodes.cpp
--- a/Engine/Plugins/Experimental/ChaosFlesh/Source/ChaosFleshNodes/Private/Dataflow/ChaosFleshBindingsNodes.cpp
+++ b/Engine/Plugins/Experimental/ChaosFlesh/Source/ChaosFleshNodes/Private/Dataflow/ChaosFleshBindingsNodes.cpp
@@ -216,9 +216,17 @@ FGenerateSurfaceBindings::Evaluate(Dataflow::FContext& Context, const FDataflowO
 					}
 				}
 			}
-
+			TArray<FString> GeometryGroupGuidsLocal = GetValue<TArray<FString>>(Context, &GeometryGroupGuidsIn);
+			TManagedArray<FString>* Guids = OutCollection.FindAttribute<FString>("Guid", FGeometryCollection::GeometryGroup);
 			for (int32 TetMeshIdx = 0; TetMeshIdx < TetrahedronStart->Num(); TetMeshIdx++)
 			{
+				if (GeometryGroupGuidsLocal.Num() && Guids)
+				{
+					if (!GeometryGroupGuidsLocal.Contains((*Guids)[TetMeshIdx]))
+					{
+						continue;
+					}
+				}
 				const int32 TetMeshStart = (*TetrahedronStart)[TetMeshIdx];
 				const int32 TetMeshCount = (*TetrahedronCount)[TetMeshIdx];
 
@@ -400,7 +408,10 @@ FGenerateSurfaceBindings::Evaluate(Dataflow::FContext& Context, const FDataflowO
 					// 
 					// Advancing front orphan reparenting
 					//
-
+					if (!MeshNeighborNodes.IsValidIndex(LOD))
+					{
+						continue;
+					}
 					const TArray<TArray<uint32>>& NeighborNodes = MeshNeighborNodes[LOD];
 					TSet<int32> OrphanSet(Orphans);
 					while (bDoOrphanReparenting && Orphans.Num())
@@ -413,6 +424,10 @@ FGenerateSurfaceBindings::Evaluate(Dataflow::FContext& Context, const FDataflowO
 						for (int32 i = 0; i < Orphans.Num(); i++)
 						{
 							int32 CurrOrphan = Orphans[i];
+							if (!NeighborNodes.IsValidIndex(CurrOrphan))
+							{
+								continue;
+							}
 							const TArray<uint32>& Neighbors = NeighborNodes[CurrOrphan];
 							int32 OrphanCount = 0;
 							int32 NonOrphanCount = 0;
diff --git a/Engine/Plugins/Experimental/ChaosFlesh/Source/ChaosFleshNodes/Private/Dataflow/ChaosFleshPositionTargetInitializationNodes.cpp b/Engine/Plugins/Experimental/ChaosFlesh/Source/ChaosFleshNodes/Private/Dataflow/ChaosFleshPositionTargetInitializationNodes.cpp
--- a/Engine/Plugins/Experimental/ChaosFlesh/Source/ChaosFleshNodes/Private/Dataflow/ChaosFleshPositionTargetInitializationNodes.cpp
+++ b/Engine/Plugins/Experimental/ChaosFlesh/Source/ChaosFleshNodes/Private/Dataflow/ChaosFleshPositionTargetInitializationNodes.cpp
@@ -282,9 +282,22 @@ void FSetVertexTetrahedraPositionTargetBindingDataflowNode::Evaluate(Dataflow::F
 					TManagedArray<int32>* TetrahedronStart = InCollection.FindAttribute<int32>(FTetrahedralCollection::TetrahedronStartAttribute, FGeometryCollection::GeometryGroup);
 					TManagedArray<int32>* TetrahedronCount = InCollection.FindAttribute<int32>(FTetrahedralCollection::TetrahedronCountAttribute, FGeometryCollection::GeometryGroup);
 					if (TetrahedronStart && TetrahedronCount)
-					{
+					{	
+						TArray<FString> GeometryGroupGuidsLocal;
+						if (FindInput(&GeometryGroupGuidsIn) && FindInput(&GeometryGroupGuidsIn)->GetConnection())
+						{
+							GeometryGroupGuidsLocal = GetValue<TArray<FString>>(Context, &GeometryGroupGuidsIn);
+						}
+						TManagedArray<FString>* Guids = InCollection.FindAttribute<FString>("Guid", FGeometryCollection::GeometryGroup);
 						for (int32 TetMeshIdx = 0; TetMeshIdx < TetrahedronStart->Num(); TetMeshIdx++)
 						{
+							if (GeometryGroupGuidsLocal.Num() && Guids)
+							{
+								if (!GeometryGroupGuidsLocal.Contains((*Guids)[TetMeshIdx]))
+								{
+									continue;
+								}
+							}
 							const int32 TetMeshStart = (*TetrahedronStart)[TetMeshIdx];
 							const int32 TetMeshCount = (*TetrahedronCount)[TetMeshIdx];
 
diff --git a/Engine/Plugins/Experimental/ChaosFlesh/Source/ChaosFleshNodes/Public/Dataflow/ChaosFleshBindingsNodes.h b/Engine/Plugins/Experimental/ChaosFlesh/Source/ChaosFleshNodes/Public/Dataflow/ChaosFleshBindingsNodes.h
--- a/Engine/Plugins/Experimental/ChaosFlesh/Source/ChaosFleshNodes/Public/Dataflow/ChaosFleshBindingsNodes.h
+++ b/Engine/Plugins/Experimental/ChaosFlesh/Source/ChaosFleshNodes/Public/Dataflow/ChaosFleshBindingsNodes.h
@@ -36,6 +36,9 @@ public:
 	UPROPERTY(EditAnywhere, Category = "Dataflow", meta = (DataflowInput, DisplayName = "SkeletalMesh"))
 	TObjectPtr<const USkeletalMesh> SkeletalMeshIn = nullptr;
 
+	UPROPERTY(meta = (DataflowInput, DisplayName = "(Optional)GeometryGroupGuidsIn"))
+	TArray<FString> GeometryGroupGuidsIn;
+
 	// Enable binding to the exterior hull of the tetrahedron mesh.
 	UPROPERTY(EditAnywhere, Category = "Dataflow", meta = (DisplayName = "DoSurfaceProjection"))
 	bool bDoSurfaceProjection = true;
@@ -55,6 +58,7 @@ public:
 		RegisterOutputConnection(&Collection, &Collection);
 		RegisterInputConnection(&StaticMeshIn);
 		RegisterInputConnection(&SkeletalMeshIn);
+		RegisterInputConnection(&GeometryGroupGuidsIn);
 	}
 
 	virtual void Evaluate(Dataflow::FContext& Context, const FDataflowOutput* Out) const override;
diff --git a/Engine/Plugins/Experimental/ChaosFlesh/Source/ChaosFleshNodes/Public/Dataflow/ChaosFleshPositionTargetInitializationNodes.h b/Engine/Plugins/Experimental/ChaosFlesh/Source/ChaosFleshNodes/Public/Dataflow/ChaosFleshPositionTargetInitializationNodes.h
--- a/Engine/Plugins/Experimental/ChaosFlesh/Source/ChaosFleshNodes/Public/Dataflow/ChaosFleshPositionTargetInitializationNodes.h
+++ b/Engine/Plugins/Experimental/ChaosFlesh/Source/ChaosFleshNodes/Public/Dataflow/ChaosFleshPositionTargetInitializationNodes.h
@@ -154,6 +154,9 @@ public:
 	UPROPERTY(meta = (DataflowInput, DisplayName = "TargetIndicesIn"))
 		TArray<int32> TargetIndicesIn;
 
+	UPROPERTY(meta = (DataflowInput, DisplayName = "GeometryGroupGuidsIn"))
+		TArray<FString> GeometryGroupGuidsIn;
+
 	UPROPERTY(EditAnywhere, Category = "Dataflow")
 		float PositionTargetStiffness = 10000.f;
 
@@ -165,6 +168,7 @@ public:
 		RegisterInputConnection(&Collection);
 		RegisterOutputConnection(&Collection, &Collection);
 		RegisterInputConnection(&TargetIndicesIn);
+		RegisterInputConnection(&GeometryGroupGuidsIn);
 	}
 
 	virtual void Evaluate(Dataflow::FContext& Context, const FDataflowOutput* Out) const override;
