ChaosFlesh : Remove the DeformableSolverAsset

diff --git a/Engine/Plugins/Experimental/ChaosFlesh/Source/ChaosFleshEditor/Private/ChaosFlesh/Asset/ActorFactoryDeformableSolver.cpp b/Engine/Plugins/Experimental/ChaosFlesh/Source/ChaosFleshEditor/Private/ChaosFlesh/Asset/ActorFactoryDeformableSolver.cpp
deleted file mode 100644
diff --git a/Engine/Plugins/Experimental/ChaosFlesh/Source/ChaosFleshEditor/Private/ChaosFlesh/Asset/ActorFactoryDeformableSolver.h b/Engine/Plugins/Experimental/ChaosFlesh/Source/ChaosFleshEditor/Private/ChaosFlesh/Asset/ActorFactoryDeformableSolver.h
deleted file mode 100644
diff --git a/Engine/Plugins/Experimental/ChaosFlesh/Source/ChaosFleshEditor/Private/ChaosFlesh/Asset/AssetTypeActions_ChaosDeformableSolverAsset.cpp b/Engine/Plugins/Experimental/ChaosFlesh/Source/ChaosFleshEditor/Private/ChaosFlesh/Asset/AssetTypeActions_ChaosDeformableSolverAsset.cpp
deleted file mode 100644
diff --git a/Engine/Plugins/Experimental/ChaosFlesh/Source/ChaosFleshEditor/Private/ChaosFlesh/Asset/ChaosAssetDeformableSolverFactory.cpp b/Engine/Plugins/Experimental/ChaosFlesh/Source/ChaosFleshEditor/Private/ChaosFlesh/Asset/ChaosAssetDeformableSolverFactory.cpp
deleted file mode 100644
diff --git a/Engine/Plugins/Experimental/ChaosFlesh/Source/ChaosFleshEditor/Private/ChaosFlesh/ChaosFleshEditorPlugin.cpp b/Engine/Plugins/Experimental/ChaosFlesh/Source/ChaosFleshEditor/Private/ChaosFlesh/ChaosFleshEditorPlugin.cpp
--- a/Engine/Plugins/Experimental/ChaosFlesh/Source/ChaosFleshEditor/Private/ChaosFlesh/ChaosFleshEditorPlugin.cpp
+++ b/Engine/Plugins/Experimental/ChaosFlesh/Source/ChaosFleshEditor/Private/ChaosFlesh/ChaosFleshEditorPlugin.cpp
@@ -7,7 +7,6 @@
 
 #include "ChaosFlesh/Cmd/ChaosFleshCommands.h"
 #include "ChaosFlesh/Asset/AssetTypeActions_FleshAsset.h"
-#include "ChaosFlesh/Asset/AssetTypeActions_ChaosDeformableSolverAsset.h"
 #include "ChaosFlesh/Asset/FleshAssetThumbnailRenderer.h"
 
 #define LOCTEXT_NAMESPACE "FleshEditor"
@@ -22,8 +21,6 @@ void IChaosFleshEditorPlugin::StartupModule()
 	IAssetTools& AssetTools = FAssetToolsModule::GetModule().Get();
 	FleshAssetActions = new FAssetTypeActions_FleshAsset();
 	AssetTools.RegisterAssetTypeActions(MakeShareable(FleshAssetActions));
-	ChaosDeformableSolverAssetActions = new FAssetTypeActions_ChaosDeformableSolver();
-	AssetTools.RegisterAssetTypeActions(MakeShareable(ChaosDeformableSolverAssetActions));
 
 
 	if (GIsEditor && !IsRunningCommandlet())
@@ -60,7 +57,6 @@ void IChaosFleshEditorPlugin::ShutdownModule()
 
 		IAssetTools& AssetTools = FAssetToolsModule::GetModule().Get();
 		AssetTools.UnregisterAssetTypeActions(FleshAssetActions->AsShared());
-		AssetTools.UnregisterAssetTypeActions(ChaosDeformableSolverAssetActions->AsShared());
 	}
 }
 
diff --git a/Engine/Plugins/Experimental/ChaosFlesh/Source/ChaosFleshEditor/Public/ChaosFlesh/Asset/AssetTypeActions_ChaosDeformableSolverAsset.h b/Engine/Plugins/Experimental/ChaosFlesh/Source/ChaosFleshEditor/Public/ChaosFlesh/Asset/AssetTypeActions_ChaosDeformableSolverAsset.h
deleted file mode 100644
diff --git a/Engine/Plugins/Experimental/ChaosFlesh/Source/ChaosFleshEditor/Public/ChaosFlesh/Asset/ChaosAssetDeformableSolverFactory.h b/Engine/Plugins/Experimental/ChaosFlesh/Source/ChaosFleshEditor/Public/ChaosFlesh/Asset/ChaosAssetDeformableSolverFactory.h
deleted file mode 100644
diff --git a/Engine/Plugins/Experimental/ChaosFlesh/Source/ChaosFleshEngine/Private/ChaosFlesh/ChaosDeformableSolverAsset.cpp b/Engine/Plugins/Experimental/ChaosFlesh/Source/ChaosFleshEngine/Private/ChaosFlesh/ChaosDeformableSolverAsset.cpp
deleted file mode 100644
diff --git a/Engine/Plugins/Experimental/ChaosFlesh/Source/ChaosFleshEngine/Public/ChaosFlesh/ChaosDeformableSolverAsset.h b/Engine/Plugins/Experimental/ChaosFlesh/Source/ChaosFleshEngine/Public/ChaosFlesh/ChaosDeformableSolverAsset.h
deleted file mode 100644
