ChaosFlesh : Add icon for the deformable solver in the editor.

diff --git a/Engine/Plugins/Experimental/ChaosFlesh/Source/ChaosFleshEngine/Private/ChaosFlesh/ChaosDeformableSolverActor.cpp b/Engine/Plugins/Experimental/ChaosFlesh/Source/ChaosFleshEngine/Private/ChaosFlesh/ChaosDeformableSolverActor.cpp
--- a/Engine/Plugins/Experimental/ChaosFlesh/Source/ChaosFleshEngine/Private/ChaosFlesh/ChaosDeformableSolverActor.cpp
+++ b/Engine/Plugins/Experimental/ChaosFlesh/Source/ChaosFleshEngine/Private/ChaosFlesh/ChaosDeformableSolverActor.cpp
@@ -3,6 +3,8 @@
 #include "ChaosFlesh/ChaosDeformableSolverActor.h"
 
 #include "ChaosFlesh/ChaosDeformableSolverComponent.h"
+#include "UObject/ConstructorHelpers.h"
+#include "Engine/Texture2D.h"
 
 DEFINE_LOG_CATEGORY_STATIC(LogChaosDeformableSolverActor, Log, All);
 
@@ -12,6 +14,51 @@ ADeformableSolverActor::ADeformableSolverActor(const FObjectInitializer& ObjectI
 	//UE_LOG(ADeformableSolverActor, Verbose, TEXT("ADeformableSolverActor::ADeformableSolverActor()"));
 	SolverComponent = CreateDefaultSubobject<UDeformableSolverComponent>(TEXT("DeformableSolverComponent0"));
 	RootComponent = SolverComponent;
+	RootComponent->Mobility = EComponentMobility::Static;
+	CreateBillboardIcon(ObjectInitializer);
+}
+
+
+void ADeformableSolverActor::CreateBillboardIcon(const FObjectInitializer& ObjectInitializer)
+{
+	/*
+	* Display icon in the editor
+	*/
+	// Structure to hold one-time initialization
+	struct FConstructorStatics
+	{
+		// A helper class object we use to find target UTexture2D object in resource package
+		ConstructorHelpers::FObjectFinderOptional<UTexture2D> SolverTextureObject;
+
+		// Icon sprite category name
+		FName ID_Solver;
+
+		// Icon sprite display name
+		FText NAME_Solver;
+
+		FConstructorStatics()
+			// Use helper class object to find the texture
+			// "/Engine/EditorResources/S_Solver" is resource path
+			: SolverTextureObject(TEXT("/Engine/EditorResources/S_Solver.S_Solver"))
+			, ID_Solver(TEXT("Solver"))
+			, NAME_Solver(NSLOCTEXT("SpriteCategory", "Solver", "Solver"))
+		{
+		}
+	};
+	static FConstructorStatics ConstructorStatics;
+
+
+#if WITH_EDITORONLY_DATA
+	SpriteComponent = ObjectInitializer.CreateEditorOnlyDefaultSubobject<UBillboardComponent>(this, TEXT("Sprite"));
+	if (SpriteComponent)
+	{
+		SpriteComponent->Sprite = ConstructorStatics.SolverTextureObject.Get();		// Get the sprite texture from helper class object
+		SpriteComponent->SpriteInfo.Category = ConstructorStatics.ID_Solver;		// Assign sprite category name
+		SpriteComponent->SpriteInfo.DisplayName = ConstructorStatics.NAME_Solver;	// Assign sprite display name
+		SpriteComponent->AttachToComponent(RootComponent, FAttachmentTransformRules::KeepRelativeTransform);
+		SpriteComponent->Mobility = EComponentMobility::Static;
+	}
+#endif // WITH_EDITORONLY_DATA
 }
 
 
diff --git a/Engine/Plugins/Experimental/ChaosFlesh/Source/ChaosFleshEngine/Public/ChaosFlesh/ChaosDeformableSolverActor.h b/Engine/Plugins/Experimental/ChaosFlesh/Source/ChaosFleshEngine/Public/ChaosFlesh/ChaosDeformableSolverActor.h
--- a/Engine/Plugins/Experimental/ChaosFlesh/Source/ChaosFleshEngine/Public/ChaosFlesh/ChaosDeformableSolverActor.h
+++ b/Engine/Plugins/Experimental/ChaosFlesh/Source/ChaosFleshEngine/Public/ChaosFlesh/ChaosDeformableSolverActor.h
@@ -3,6 +3,7 @@
 #pragma once
 
 #include "CoreMinimal.h"
+#include "Components/BillboardComponent.h"
 #include "GameFramework/Actor.h"
 #include "UObject/ObjectMacros.h"
 #include "UObject/ObjectPtr.h"
@@ -24,4 +25,13 @@ public:
 	UPROPERTY(Category = "Chaos Deformable", VisibleAnywhere, BlueprintReadOnly, meta = (AllowPrivateAccess = "true"))
 	TObjectPtr<UDeformableSolverComponent> SolverComponent = nullptr;
 
+	/*
+	* Display icon in the editor
+	*/
+	UPROPERTY()
+	TObjectPtr<UBillboardComponent> SpriteComponent;
+
+private:
+	void CreateBillboardIcon(const FObjectInitializer& ObjectInitializer);
+
 };
\ No newline at end of file
