ChaosFlesh : UProperty updates
- Added specializations to Flesh Actor and Flesh Components to hide unused physics assets.
- Move all Deformable properties into the Physics Category.

diff --git a/Engine/Plugins/Experimental/ChaosFlesh/Source/ChaosFleshEditor/Private/ChaosFlesh/Asset/FleshDeformableInterfaceDetails.cpp b/Engine/Plugins/Experimental/ChaosFlesh/Source/ChaosFleshEditor/Private/ChaosFlesh/Asset/FleshDeformableInterfaceDetails.cpp
new file mode 100644
--- /dev/null
+++ b/Engine/Plugins/Experimental/ChaosFlesh/Source/ChaosFleshEditor/Private/ChaosFlesh/Asset/FleshDeformableInterfaceDetails.cpp
@@ -0,0 +1,33 @@
+// Copyright Epic Games, Inc. All Rights Reserved.
+
+#include "FleshDeformableInterfaceDetails.h"
+
+#include "DetailLayoutBuilder.h"
+#include "Misc/AssertionMacros.h"
+#include "BodyInstanceCore.h"
+#include "DeformableInterface.h"
+
+#define LOCTEXT_NAMESPACE "DeformableInterfaceDetails"
+
+
+TSharedRef<IDetailCustomization> FDeformableInterfaceDetails::MakeInstance()
+{
+	return MakeShareable(new FDeformableInterfaceDetails);
+}
+
+void FDeformableInterfaceDetails::CustomizeDetails(IDetailLayoutBuilder& DetailBuilder)
+{
+	TArray<TWeakObjectPtr<UObject>> DetailObjects;
+	DetailBuilder.GetObjectsBeingCustomized(DetailObjects);
+	for (TWeakObjectPtr<UObject> DetailObject : DetailObjects)
+	{
+		if (IDeformableInterface* DeformableInterface = Cast<IDeformableInterface>(DetailObject.Get()))
+		{
+			DeformableInterface->CustomizeDetails(DetailBuilder);
+		}
+	}
+}
+
+
+
+#undef LOCTEXT_NAMESPACE
diff --git a/Engine/Plugins/Experimental/ChaosFlesh/Source/ChaosFleshEditor/Private/ChaosFlesh/Asset/FleshDeformableInterfaceDetails.h b/Engine/Plugins/Experimental/ChaosFlesh/Source/ChaosFleshEditor/Private/ChaosFlesh/Asset/FleshDeformableInterfaceDetails.h
new file mode 100644
--- /dev/null
+++ b/Engine/Plugins/Experimental/ChaosFlesh/Source/ChaosFleshEditor/Private/ChaosFlesh/Asset/FleshDeformableInterfaceDetails.h
@@ -0,0 +1,20 @@
+// Copyright Epic Games, Inc. All Rights Reserved.
+
+#pragma once
+
+#include "CoreMinimal.h"
+#include "IDetailCustomization.h"
+#include "IDetailCustomNodeBuilder.h"
+#include "ChaosFlesh/FleshComponent.h"
+
+
+
+class FDeformableInterfaceDetails : public IDetailCustomization
+{
+public:
+	/** Makes a new instance of this detail layout class for a specific detail view requesting it */
+	static TSharedRef<IDetailCustomization> MakeInstance();
+
+	/** IDetailCustomization interface */
+	virtual void CustomizeDetails( IDetailLayoutBuilder& DetailBuilder ) override;
+};
diff --git a/Engine/Plugins/Experimental/ChaosFlesh/Source/ChaosFleshEditor/Private/ChaosFlesh/ChaosFleshEditorPlugin.cpp b/Engine/Plugins/Experimental/ChaosFlesh/Source/ChaosFleshEditor/Private/ChaosFlesh/ChaosFleshEditorPlugin.cpp
--- a/Engine/Plugins/Experimental/ChaosFlesh/Source/ChaosFleshEditor/Private/ChaosFlesh/ChaosFleshEditorPlugin.cpp
+++ b/Engine/Plugins/Experimental/ChaosFlesh/Source/ChaosFleshEditor/Private/ChaosFlesh/ChaosFleshEditorPlugin.cpp
@@ -2,12 +2,17 @@
 
 
 #include "ChaosFlesh/ChaosFleshEditorPlugin.h"
-#include "ChaosFlesh/FleshAsset.h"
-#include "Editor/FleshEditorStyle.h"
 
-#include "ChaosFlesh/Cmd/ChaosFleshCommands.h"
 #include "ChaosFlesh/Asset/AssetTypeActions_FleshAsset.h"
+#include "ChaosFlesh/Asset/FleshDeformableInterfaceDetails.h"
 #include "ChaosFlesh/Asset/FleshAssetThumbnailRenderer.h"
+#include "ChaosFlesh/ChaosDeformableCollisionsActor.h"
+#include "ChaosFlesh/ChaosDeformableSolverComponent.h"
+#include "ChaosFlesh/Cmd/ChaosFleshCommands.h"
+#include "ChaosFlesh/FleshActor.h"
+#include "ChaosFlesh/FleshAsset.h"
+#include "ChaosFlesh/FleshComponent.h"
+#include "Editor/FleshEditorStyle.h"
 
 #define LOCTEXT_NAMESPACE "FleshEditor"
 
@@ -19,10 +24,10 @@ void IChaosFleshEditorPlugin::StartupModule()
 	FChaosFleshEditorStyle::Get();
 
 	IAssetTools& AssetTools = FAssetToolsModule::GetModule().Get();
+
 	FleshAssetActions = new FAssetTypeActions_FleshAsset();
 	AssetTools.RegisterAssetTypeActions(MakeShareable(FleshAssetActions));
 
-
 	if (GIsEditor && !IsRunningCommandlet())
 	{
 		EditorCommands.Add(IConsoleManager::Get().RegisterConsoleCommand(
@@ -47,6 +52,29 @@ void IChaosFleshEditorPlugin::StartupModule()
 
 	UThumbnailManager::Get().RegisterCustomRenderer(UFleshAsset::StaticClass(), UFleshAssetThumbnailRenderer::StaticClass());
 
+	// register details customization
+	FPropertyEditorModule& PropertyModule = FModuleManager::LoadModuleChecked<FPropertyEditorModule>("PropertyEditor");
+
+	PropertyModule.RegisterCustomClassLayout(
+		ADeformableCollisionsActor::StaticClass()->GetFName(),
+		FOnGetDetailCustomizationInstance::CreateStatic(&FDeformableInterfaceDetails::MakeInstance));
+
+	PropertyModule.RegisterCustomClassLayout(
+		ADeformableSolverActor::StaticClass()->GetFName(),
+		FOnGetDetailCustomizationInstance::CreateStatic(&FDeformableInterfaceDetails::MakeInstance));
+
+	PropertyModule.RegisterCustomClassLayout(
+		AFleshActor::StaticClass()->GetFName(),
+		FOnGetDetailCustomizationInstance::CreateStatic(&FDeformableInterfaceDetails::MakeInstance));
+
+	PropertyModule.RegisterCustomClassLayout(
+		UDeformablePhysicsComponent::StaticClass()->GetFName(),
+		FOnGetDetailCustomizationInstance::CreateStatic(&FDeformableInterfaceDetails::MakeInstance));
+
+	PropertyModule.RegisterCustomClassLayout(
+		UDeformableSolverComponent::StaticClass()->GetFName(),
+		FOnGetDetailCustomizationInstance::CreateStatic(&FDeformableInterfaceDetails::MakeInstance));
+
 }
 
 void IChaosFleshEditorPlugin::ShutdownModule()
diff --git a/Engine/Plugins/Experimental/ChaosFlesh/Source/ChaosFleshEngine/Private/ChaosFlesh/ChaosDeformableCollisionsActor.cpp b/Engine/Plugins/Experimental/ChaosFlesh/Source/ChaosFleshEngine/Private/ChaosFlesh/ChaosDeformableCollisionsActor.cpp
--- a/Engine/Plugins/Experimental/ChaosFlesh/Source/ChaosFleshEngine/Private/ChaosFlesh/ChaosDeformableCollisionsActor.cpp
+++ b/Engine/Plugins/Experimental/ChaosFlesh/Source/ChaosFleshEngine/Private/ChaosFlesh/ChaosDeformableCollisionsActor.cpp
@@ -16,6 +16,14 @@ ADeformableCollisionsActor::ADeformableCollisionsActor(const FObjectInitializer&
 	PrimaryActorTick.bCanEverTick = false;
 }
 
+#if WITH_EDITOR
+void ADeformableCollisionsActor::CustomizeDetails(IDetailLayoutBuilder& DetailBuilder)
+{
+	//TSharedPtr<IPropertyHandle> bReplicatePhysicsProperty = DetailBuilder.GetProperty(GET_MEMBER_NAME_CHECKED(AFleshActor, bAsyncPhysicsTickEnabled), AActor::StaticClass());
+	//bReplicatePhysicsProperty->MarkHiddenByCustomization();
+}
+#endif
+
 void ADeformableCollisionsActor::EnableSimulation(ADeformableSolverActor* InActor)
 {
 	if (InActor)
diff --git a/Engine/Plugins/Experimental/ChaosFlesh/Source/ChaosFleshEngine/Private/ChaosFlesh/ChaosDeformablePhysicsComponent.cpp b/Engine/Plugins/Experimental/ChaosFlesh/Source/ChaosFleshEngine/Private/ChaosFlesh/ChaosDeformablePhysicsComponent.cpp
--- a/Engine/Plugins/Experimental/ChaosFlesh/Source/ChaosFleshEngine/Private/ChaosFlesh/ChaosDeformablePhysicsComponent.cpp
+++ b/Engine/Plugins/Experimental/ChaosFlesh/Source/ChaosFleshEngine/Private/ChaosFlesh/ChaosDeformablePhysicsComponent.cpp
@@ -7,6 +7,56 @@
 #include "ChaosFlesh/FleshComponent.h"
 
 
+#if WITH_EDITOR
+void UDeformablePhysicsComponent::CustomizeDetails(IDetailLayoutBuilder& DetailBuilder)
+{
+	TSharedPtr<IPropertyHandle>bApplyImpulseOnDamageProperty = DetailBuilder.GetProperty(GET_MEMBER_NAME_CHECKED(UFleshComponent, bApplyImpulseOnDamage), UPrimitiveComponent::StaticClass());
+	bApplyImpulseOnDamageProperty->MarkHiddenByCustomization();
+
+	TSharedPtr<IPropertyHandle> bIgnoreRadialImpulseProperty = DetailBuilder.GetProperty(GET_MEMBER_NAME_CHECKED(UFleshComponent, bIgnoreRadialImpulse), UPrimitiveComponent::StaticClass());
+	bIgnoreRadialImpulseProperty->MarkHiddenByCustomization();
+
+	TSharedPtr<IPropertyHandle> bIgnoreRadialForceProperty = DetailBuilder.GetProperty(GET_MEMBER_NAME_CHECKED(UFleshComponent, bIgnoreRadialForce), UPrimitiveComponent::StaticClass());
+	bIgnoreRadialForceProperty->MarkHiddenByCustomization();
+
+	TSharedPtr<IPropertyHandle> bReplicatePhysicsToAutonomousProxyProperty = DetailBuilder.GetProperty(GET_MEMBER_NAME_CHECKED(UFleshComponent, bReplicatePhysicsToAutonomousProxy), UPrimitiveComponent::StaticClass());
+	bReplicatePhysicsToAutonomousProxyProperty->MarkHiddenByCustomization();
+
+	TSharedPtr<IPropertyHandle> ShouldUpdatePhysicsVolumnProperty = DetailBuilder.GetProperty("bShouldUpdatePhysicsVolume", USceneComponent::StaticClass());
+	ShouldUpdatePhysicsVolumnProperty->MarkHiddenByCustomization();
+
+	//
+	// Some BodyInstance properties are managed through FBodyInstanceCustomizationHelper
+	// For example, this is not needed here as its overwritten in the helper.
+	//
+	// TSharedPtr<IPropertyHandle> BodyInstance_bSimulatePhysicsProperty = DetailBuilder.GetProperty(GET_MEMBER_NAME_CHECKED(UFleshComponent, BodyInstance.bSimulatePhysics), UPrimitiveComponent::StaticClass());
+	// BodyInstance_bSimulatePhysicsProperty->MarkHiddenByCustomization();
+	//
+	// Properties not managed in the Helper are hidden here:
+
+	TSharedPtr<IPropertyHandle> SleepFamilyProperty = DetailBuilder.GetProperty(GET_MEMBER_NAME_CHECKED(UFleshComponent, BodyInstance.SleepFamily), UPrimitiveComponent::StaticClass());
+	SleepFamilyProperty->MarkHiddenByCustomization();
+
+	TSharedPtr<IPropertyHandle> InertiaTensorScaleProperty = DetailBuilder.GetProperty(GET_MEMBER_NAME_CHECKED(UFleshComponent, BodyInstance.InertiaTensorScale), UPrimitiveComponent::StaticClass());
+	InertiaTensorScaleProperty->MarkHiddenByCustomization();
+
+	TSharedPtr<IPropertyHandle> CustomSleepThresholdMultiplierProperty = DetailBuilder.GetProperty(GET_MEMBER_NAME_CHECKED(UFleshComponent, BodyInstance.CustomSleepThresholdMultiplier), UPrimitiveComponent::StaticClass());
+	CustomSleepThresholdMultiplierProperty->MarkHiddenByCustomization();
+
+	TSharedPtr<IPropertyHandle> StabilizationThresholdMultiplierProperty = DetailBuilder.GetProperty(GET_MEMBER_NAME_CHECKED(UFleshComponent, BodyInstance.StabilizationThresholdMultiplier), UPrimitiveComponent::StaticClass());
+	StabilizationThresholdMultiplierProperty->MarkHiddenByCustomization();
+
+	TSharedPtr<IPropertyHandle> GenerateWakeEventsProperty = DetailBuilder.GetProperty(GET_MEMBER_NAME_CHECKED(UFleshComponent, BodyInstance.bGenerateWakeEvents), UPrimitiveComponent::StaticClass());
+	GenerateWakeEventsProperty->MarkHiddenByCustomization();
+
+	TSharedPtr<IPropertyHandle> UpdateMassWhenScaleChangesProperty = DetailBuilder.GetProperty(GET_MEMBER_NAME_CHECKED(UFleshComponent, BodyInstance.bUpdateMassWhenScaleChanges), UPrimitiveComponent::StaticClass());
+	UpdateMassWhenScaleChangesProperty->MarkHiddenByCustomization();
+
+	TSharedPtr<IPropertyHandle> MaxAngularVelocityProperty = DetailBuilder.GetProperty(GET_MEMBER_NAME_CHECKED(UFleshComponent, BodyInstance.MaxAngularVelocity), UPrimitiveComponent::StaticClass());
+	MaxAngularVelocityProperty->MarkHiddenByCustomization();
+}
+#endif
+
 void UDeformablePhysicsComponent::OnCreatePhysicsState()
 {
 	Super::OnCreatePhysicsState();
diff --git a/Engine/Plugins/Experimental/ChaosFlesh/Source/ChaosFleshEngine/Private/ChaosFlesh/ChaosDeformableSolverActor.cpp b/Engine/Plugins/Experimental/ChaosFlesh/Source/ChaosFleshEngine/Private/ChaosFlesh/ChaosDeformableSolverActor.cpp
--- a/Engine/Plugins/Experimental/ChaosFlesh/Source/ChaosFleshEngine/Private/ChaosFlesh/ChaosDeformableSolverActor.cpp
+++ b/Engine/Plugins/Experimental/ChaosFlesh/Source/ChaosFleshEngine/Private/ChaosFlesh/ChaosDeformableSolverActor.cpp
@@ -3,8 +3,8 @@
 #include "ChaosFlesh/ChaosDeformableSolverActor.h"
 
 #include "ChaosFlesh/ChaosDeformableSolverComponent.h"
-#include "UObject/ConstructorHelpers.h"
 #include "Engine/Texture2D.h"
+#include "UObject/ConstructorHelpers.h"
 
 DEFINE_LOG_CATEGORY_STATIC(LogChaosDeformableSolverActor, Log, All);
 
@@ -18,6 +18,14 @@ ADeformableSolverActor::ADeformableSolverActor(const FObjectInitializer& ObjectI
 	CreateBillboardIcon(ObjectInitializer);
 }
 
+#if WITH_EDITOR
+void ADeformableSolverActor::CustomizeDetails(IDetailLayoutBuilder& DetailBuilder)
+{
+	//TSharedPtr<IPropertyHandle> bReplicatePhysicsProperty = DetailBuilder.GetProperty(GET_MEMBER_NAME_CHECKED(AFleshActor, bAsyncPhysicsTickEnabled), AActor::StaticClass());
+	//bReplicatePhysicsProperty->MarkHiddenByCustomization();
+}
+
+#endif
 
 void ADeformableSolverActor::CreateBillboardIcon(const FObjectInitializer& ObjectInitializer)
 {
diff --git a/Engine/Plugins/Experimental/ChaosFlesh/Source/ChaosFleshEngine/Private/ChaosFlesh/ChaosDeformableSolverComponent.cpp b/Engine/Plugins/Experimental/ChaosFlesh/Source/ChaosFleshEngine/Private/ChaosFlesh/ChaosDeformableSolverComponent.cpp
--- a/Engine/Plugins/Experimental/ChaosFlesh/Source/ChaosFleshEngine/Private/ChaosFlesh/ChaosDeformableSolverComponent.cpp
+++ b/Engine/Plugins/Experimental/ChaosFlesh/Source/ChaosFleshEngine/Private/ChaosFlesh/ChaosDeformableSolverComponent.cpp
@@ -40,6 +40,14 @@ UDeformableSolverComponent::~UDeformableSolverComponent()
 {
 }
 
+#if WITH_EDITOR
+void UDeformableSolverComponent::CustomizeDetails(IDetailLayoutBuilder& DetailBuilder)
+{
+	//TSharedPtr<IPropertyHandle> bReplicatePhysicsProperty = DetailBuilder.GetProperty(GET_MEMBER_NAME_CHECKED(AFleshActor, bAsyncPhysicsTickEnabled), AActor::StaticClass());
+	//bReplicatePhysicsProperty->MarkHiddenByCustomization();
+}
+#endif
+
 void UDeformableSolverComponent::UpdateTickGroup()
 {
 	//
diff --git a/Engine/Plugins/Experimental/ChaosFlesh/Source/ChaosFleshEngine/Private/ChaosFlesh/FleshActor.cpp b/Engine/Plugins/Experimental/ChaosFlesh/Source/ChaosFleshEngine/Private/ChaosFlesh/FleshActor.cpp
--- a/Engine/Plugins/Experimental/ChaosFlesh/Source/ChaosFleshEngine/Private/ChaosFlesh/FleshActor.cpp
+++ b/Engine/Plugins/Experimental/ChaosFlesh/Source/ChaosFleshEngine/Private/ChaosFlesh/FleshActor.cpp
@@ -5,9 +5,6 @@
 #include "ChaosFlesh/ChaosDeformableSolverComponent.h"
 #include "ChaosFlesh/ChaosDeformableSolverActor.h"
 
-//#include "ChaosFlesh/PB.h"
-
-
 #include UE_INLINE_GENERATED_CPP_BY_NAME(FleshActor)
 
 DEFINE_LOG_CATEGORY_STATIC(AFleshLogging, Log, All);
@@ -96,6 +93,13 @@ bool AFleshActor::GetReferencedContentObjects(TArray<UObject*>& Objects) const
 	}
 	return true;
 }
+
+void AFleshActor::CustomizeDetails(IDetailLayoutBuilder& DetailBuilder)
+{
+	TSharedPtr<IPropertyHandle> bReplicatePhysicsProperty = DetailBuilder.GetProperty(GET_MEMBER_NAME_CHECKED(AFleshActor, bAsyncPhysicsTickEnabled), AActor::StaticClass());
+	bReplicatePhysicsProperty->MarkHiddenByCustomization();
+}
+
 #endif
 
 
diff --git a/Engine/Plugins/Experimental/ChaosFlesh/Source/ChaosFleshEngine/Public/ChaosFlesh/ChaosDeformableCollisionsActor.h b/Engine/Plugins/Experimental/ChaosFlesh/Source/ChaosFleshEngine/Public/ChaosFlesh/ChaosDeformableCollisionsActor.h
--- a/Engine/Plugins/Experimental/ChaosFlesh/Source/ChaosFleshEngine/Public/ChaosFlesh/ChaosDeformableCollisionsActor.h
+++ b/Engine/Plugins/Experimental/ChaosFlesh/Source/ChaosFleshEngine/Public/ChaosFlesh/ChaosDeformableCollisionsActor.h
@@ -2,12 +2,13 @@
 
 #pragma once
 
+#include "ChaosFlesh/ChaosDeformableCollisionsComponent.h"
+#include "ChaosFlesh/ChaosDeformableSolverActor.h"
 #include "CoreMinimal.h"
+#include "DeformableInterface.h"
 #include "GameFramework/Actor.h"
 #include "UObject/ObjectMacros.h"
 #include "UObject/ObjectPtr.h"
-#include "ChaosFlesh/ChaosDeformableCollisionsComponent.h"
-#include "ChaosFlesh/ChaosDeformableSolverActor.h"
 
 #include "ChaosDeformableCollisionsActor.generated.h"
 
@@ -15,7 +16,7 @@ class AStaticMeshActor;
 
 
 UCLASS()
-class CHAOSFLESHENGINE_API ADeformableCollisionsActor : public AActor
+class CHAOSFLESHENGINE_API ADeformableCollisionsActor : public AActor, public IDeformableInterface
 {
 	GENERATED_UCLASS_BODY()
 
@@ -24,14 +25,14 @@ public:
 	UFUNCTION(BlueprintCallable, Category = "Physics")
 	void EnableSimulation(ADeformableSolverActor* Actor);
 
-	UPROPERTY(VisibleAnywhere, BlueprintReadOnly, Category = "Chaos Deformable")
+	UPROPERTY(VisibleAnywhere, BlueprintReadOnly, Category = "Physics")
 	TObjectPtr<UDeformableCollisionsComponent> DeformableCollisionsComponent;
 	UDeformableCollisionsComponent* GetCollisionsComponent() const { return DeformableCollisionsComponent; }
 
-	UPROPERTY(EditAnywhere, BlueprintReadWrite, Category = "Chaos Deformable")
+	UPROPERTY(EditAnywhere, BlueprintReadWrite, Category = "Physics")
 	TObjectPtr<ADeformableSolverActor> PrimarySolver;
 
-	UPROPERTY(EditAnywhere, BlueprintReadWrite, Category = "Chaos Deformable")
+	UPROPERTY(EditAnywhere, BlueprintReadWrite, Category = "Physics")
 	TArray<TObjectPtr<AStaticMeshActor>> StaticCollisions;
 
 #if WITH_EDITOR
@@ -41,5 +42,6 @@ public:
 	ADeformableSolverActor* PreEditChangePrimarySolver = nullptr;
 	virtual void PreEditChange(FProperty* PropertyThatWillChange) override;
 	virtual void PostEditChangeProperty(struct FPropertyChangedEvent& PropertyChangedEvent) override;
+	virtual void CustomizeDetails(IDetailLayoutBuilder& DetailBuilder) override;
 #endif
 };
diff --git a/Engine/Plugins/Experimental/ChaosFlesh/Source/ChaosFleshEngine/Public/ChaosFlesh/ChaosDeformableCollisionsComponent.h b/Engine/Plugins/Experimental/ChaosFlesh/Source/ChaosFleshEngine/Public/ChaosFlesh/ChaosDeformableCollisionsComponent.h
--- a/Engine/Plugins/Experimental/ChaosFlesh/Source/ChaosFleshEngine/Public/ChaosFlesh/ChaosDeformableCollisionsComponent.h
+++ b/Engine/Plugins/Experimental/ChaosFlesh/Source/ChaosFleshEngine/Public/ChaosFlesh/ChaosDeformableCollisionsComponent.h
@@ -37,7 +37,7 @@ public:
 	UFUNCTION(BlueprintCallable, Category = "Physics")
 	void RemoveStaticMeshComponent(UStaticMeshComponent* StaticMeshComponent);
 
-	UPROPERTY(BlueprintReadOnly, Category = "Chaos Deformable")
+	UPROPERTY(BlueprintReadOnly, Category = "Physics")
 	TArray<TObjectPtr<UStaticMeshComponent>>  CollisionBodies;
 
 protected:
diff --git a/Engine/Plugins/Experimental/ChaosFlesh/Source/ChaosFleshEngine/Public/ChaosFlesh/ChaosDeformableGameplayComponent.h b/Engine/Plugins/Experimental/ChaosFlesh/Source/ChaosFleshEngine/Public/ChaosFlesh/ChaosDeformableGameplayComponent.h
--- a/Engine/Plugins/Experimental/ChaosFlesh/Source/ChaosFleshEngine/Public/ChaosFlesh/ChaosDeformableGameplayComponent.h
+++ b/Engine/Plugins/Experimental/ChaosFlesh/Source/ChaosFleshEngine/Public/ChaosFlesh/ChaosDeformableGameplayComponent.h
@@ -67,11 +67,11 @@ public:
 	void DetectEnvironmentCollisions(const int32 MaxNumTests = 100, const bool bTestDownOnly = true, const float TestRange = 0.0, const ECollisionChannel CollisionChannel = ECollisionChannel::ECC_PhysicsBody);
 
 
-	UPROPERTY(EditAnywhere, Category = "ChaosDeformable")
+	UPROPERTY(EditAnywhere, Category = "Physics")
 	FRigBoundRayCasts RigBoundRayCasts;
 
 	/** Objects to skip if hit during \c DetectEnvironmentCollisions(). */
-	UPROPERTY(EditAnywhere, Category = "ChaosDeformable")
+	UPROPERTY(EditAnywhere, Category = "Physics")
 	TArray<TWeakObjectPtr<UPrimitiveComponent>> EnvironmentCollisionsSkipList;
 
 private:
diff --git a/Engine/Plugins/Experimental/ChaosFlesh/Source/ChaosFleshEngine/Public/ChaosFlesh/ChaosDeformablePhysicsComponent.h b/Engine/Plugins/Experimental/ChaosFlesh/Source/ChaosFleshEngine/Public/ChaosFlesh/ChaosDeformablePhysicsComponent.h
--- a/Engine/Plugins/Experimental/ChaosFlesh/Source/ChaosFleshEngine/Public/ChaosFlesh/ChaosDeformablePhysicsComponent.h
+++ b/Engine/Plugins/Experimental/ChaosFlesh/Source/ChaosFleshEngine/Public/ChaosFlesh/ChaosDeformablePhysicsComponent.h
@@ -5,6 +5,7 @@
 #include "Chaos/Deformable/ChaosDeformableSolver.h"
 #include "ChaosFlesh/ChaosDeformableSolverThreading.h"
 #include "Components/MeshComponent.h"
+#include "DeformableInterface.h"
 #include "UObject/ObjectMacros.h"
 #include "ProceduralMeshComponent.h"
 #include "ChaosDeformablePhysicsComponent.generated.h"
@@ -16,7 +17,7 @@ class UDeformableSolverComponent;
 *	UDeformablePhysicsComponent
 */
 UCLASS(meta = (BlueprintSpawnableComponent))
-class CHAOSFLESHENGINE_API UDeformablePhysicsComponent : public UPrimitiveComponent
+class CHAOSFLESHENGINE_API UDeformablePhysicsComponent : public UPrimitiveComponent, public IDeformableInterface
 {
 	GENERATED_UCLASS_BODY()
 
@@ -50,30 +51,34 @@ public:
 	UDeformableSolverComponent* GetDeformableSolver();
 	const UDeformableSolverComponent* GetDeformableSolver() const;
 	
-	UPROPERTY(EditAnywhere, BlueprintReadWrite, Category = "ChaosDeformable")
+	UPROPERTY(EditAnywhere, BlueprintReadWrite, Category = "Physics")
 	TObjectPtr<UDeformableSolverComponent> PrimarySolverComponent;
 
-	UPROPERTY(EditAnywhere, BlueprintReadWrite, Category = "ChaosDeformable")
+	UPROPERTY(EditAnywhere, BlueprintReadWrite, Category = "Physics")
 	bool bTempEnableGravity = true;
 
-	UPROPERTY(EditAnywhere, BlueprintReadWrite, Category = "ChaosDeformable", meta = (ClampMin = "0.0", ClampMax = "1.0", UIMin = "0.0", UIMax = "1.0"))
+	UPROPERTY(EditAnywhere, BlueprintReadWrite, Category = "Physics", meta = (ClampMin = "0.0", ClampMax = "1.0", UIMin = "0.0", UIMax = "1.0"))
 	float DampingMultiplier = 1.f;
 
-	UPROPERTY(EditAnywhere, BlueprintReadWrite, Category = "ChaosDeformable", meta = (ClampMin = "0.0", ClampMax = "1.0", UIMin = "0.0", UIMax = "1.0"))
+	UPROPERTY(EditAnywhere, BlueprintReadWrite, Category = "Physics", meta = (ClampMin = "0.0", ClampMax = "1.0", UIMin = "0.0", UIMax = "1.0"))
 	float StiffnessMultiplier = 1.f;
 
-	UPROPERTY(EditAnywhere, BlueprintReadWrite, Category = "ChaosDeformable", meta = (ClampMin = "0.0", ClampMax = "1.0", UIMin = "0.0", UIMax = "1.0"))
+	UPROPERTY(EditAnywhere, BlueprintReadWrite, Category = "Physics", meta = (ClampMin = "0.0", ClampMax = "1.0", UIMin = "0.0", UIMax = "1.0"))
 	float MassMultiplier = 1.f;
 
-	UPROPERTY(EditAnywhere, BlueprintReadWrite, Category = "ChaosDeformable", meta = (ClampMin = "0.0", ClampMax = "1.0", UIMin = "0.0", UIMax = "1.0"))
+	UPROPERTY(EditAnywhere, BlueprintReadWrite, Category = "Physics", meta = (ClampMin = "0.0", ClampMax = "1.0", UIMin = "0.0", UIMax = "1.0"))
 	float IncompressibilityMultiplier = 1.f;
 
-	UPROPERTY(EditAnywhere, BlueprintReadWrite, Category = "ChaosDeformable", meta = (ClampMin = "0.0", ClampMax = "1.0", UIMin = "0.0", UIMax = "1.0"))
+	UPROPERTY(EditAnywhere, BlueprintReadWrite, Category = "Physics", meta = (ClampMin = "0.0", ClampMax = "1.0", UIMin = "0.0", UIMax = "1.0"))
 	float InflationMultiplier = 1.f;
 
 	const FThreadingProxy* GetPhysicsProxy() const { return PhysicsProxy; }
 	FThreadingProxy* GetPhysicsProxy() { return PhysicsProxy; }
 
+#if WITH_EDITOR
+	virtual void CustomizeDetails(IDetailLayoutBuilder& DetailBuilder) override;
+#endif
+
 protected:
 	FThreadingProxy* PhysicsProxy = nullptr;
 
diff --git a/Engine/Plugins/Experimental/ChaosFlesh/Source/ChaosFleshEngine/Public/ChaosFlesh/ChaosDeformableSolverActor.h b/Engine/Plugins/Experimental/ChaosFlesh/Source/ChaosFleshEngine/Public/ChaosFlesh/ChaosDeformableSolverActor.h
--- a/Engine/Plugins/Experimental/ChaosFlesh/Source/ChaosFleshEngine/Public/ChaosFlesh/ChaosDeformableSolverActor.h
+++ b/Engine/Plugins/Experimental/ChaosFlesh/Source/ChaosFleshEngine/Public/ChaosFlesh/ChaosDeformableSolverActor.h
@@ -4,6 +4,7 @@
 
 #include "CoreMinimal.h"
 #include "Components/BillboardComponent.h"
+#include "DeformableInterface.h"
 #include "GameFramework/Actor.h"
 #include "UObject/ObjectMacros.h"
 #include "UObject/ObjectPtr.h"
@@ -13,7 +14,7 @@
 class UDeformableSolverComponent;
 
 UCLASS()
-class CHAOSFLESHENGINE_API ADeformableSolverActor : public AActor
+class CHAOSFLESHENGINE_API ADeformableSolverActor : public AActor, public IDeformableInterface
 {
 	GENERATED_UCLASS_BODY()
 
@@ -22,7 +23,7 @@ public:
 	UDeformableSolverComponent* GetDeformableSolverComponent() { return SolverComponent; }
 	const UDeformableSolverComponent* GetDeformableSolverComponent() const { return SolverComponent; }
 
-	UPROPERTY(Category = "Chaos Deformable", VisibleAnywhere, BlueprintReadOnly, meta = (AllowPrivateAccess = "true"))
+	UPROPERTY(Category = "Physics", VisibleAnywhere, BlueprintReadOnly, meta = (AllowPrivateAccess = "true"))
 	TObjectPtr<UDeformableSolverComponent> SolverComponent = nullptr;
 
 	/*
@@ -31,6 +32,10 @@ public:
 	UPROPERTY()
 	TObjectPtr<UBillboardComponent> SpriteComponent;
 
+#if WITH_EDITOR
+	virtual void CustomizeDetails(IDetailLayoutBuilder& DetailBuilder) override;
+#endif
+
 private:
 	void CreateBillboardIcon(const FObjectInitializer& ObjectInitializer);
 
diff --git a/Engine/Plugins/Experimental/ChaosFlesh/Source/ChaosFleshEngine/Public/ChaosFlesh/ChaosDeformableSolverComponent.h b/Engine/Plugins/Experimental/ChaosFlesh/Source/ChaosFleshEngine/Public/ChaosFlesh/ChaosDeformableSolverComponent.h
--- a/Engine/Plugins/Experimental/ChaosFlesh/Source/ChaosFleshEngine/Public/ChaosFlesh/ChaosDeformableSolverComponent.h
+++ b/Engine/Plugins/Experimental/ChaosFlesh/Source/ChaosFleshEngine/Public/ChaosFlesh/ChaosDeformableSolverComponent.h
@@ -7,6 +7,7 @@
 #include "ChaosFlesh/ChaosDeformableSolverThreading.h"
 #include "ChaosFlesh/ChaosDeformableTypes.h"
 #include "ChaosFlesh/FleshComponent.h"
+#include "DeformableInterface.h"
 #include "Components/SceneComponent.h"
 #include "UObject/ObjectMacros.h"
 
@@ -19,7 +20,7 @@ class UDeformableCollisionsComponent;
 *	UDeformableSolverComponent
 */
 UCLASS(meta = (BlueprintSpawnableComponent))
-class CHAOSFLESHENGINE_API UDeformableSolverComponent : public USceneComponent
+class CHAOSFLESHENGINE_API UDeformableSolverComponent : public USceneComponent, public IDeformableInterface
 {
 	GENERATED_UCLASS_BODY()
 
@@ -53,71 +54,71 @@ public:
 
 	/* Properties */
 
-	UPROPERTY(EditAnywhere, BlueprintReadOnly, Category = "Chaos Deformable")
+	UPROPERTY(EditAnywhere, BlueprintReadOnly, Category = "Physics", meta = (EditCondition = "false"))
 	TArray< TObjectPtr<UDeformablePhysicsComponent> > DeformableComponents;
 
-	UPROPERTY(EditAnywhere, BlueprintReadOnly, Category = "Chaos Deformable")
+	UPROPERTY(EditAnywhere, BlueprintReadOnly, Category = "Physics", meta = (EditCondition = "false"))
 	TObjectPtr<UDeformableCollisionsComponent> CollisionComponent;
 
 	/** ObjectType defines how to initialize the rigid objects state, Kinematic, Sleeping, Dynamic. */
-	UPROPERTY(EditAnywhere, BlueprintReadOnly , Category = "Chaos Deformable")
+	UPROPERTY(EditAnywhere, BlueprintReadOnly , Category = "Physics")
 	EDeformableExecutionModel ExecutionModel = EDeformableExecutionModel::Chaos_Deformable_DuringPhysics;
 
-	UPROPERTY(EditAnywhere, Category = "Chaos Deformable")
+	UPROPERTY(EditAnywhere, Category = "Physics")
 		int32 NumSubSteps = 2;
 
-	UPROPERTY(EditAnywhere, Category = "Chaos Deformable")
+	UPROPERTY(EditAnywhere, Category = "Physics")
 		int32 NumSolverIterations = 5;
 
-	UPROPERTY(EditAnywhere, Category = "Chaos Deformable")
+	UPROPERTY(EditAnywhere, Category = "Physics")
 		bool FixTimeStep = false;
 
-	UPROPERTY(EditAnywhere, Category = "Chaos Deformable")
+	UPROPERTY(EditAnywhere, Category = "Physics")
 		float TimeStepSize = 0.05;
 
-	UPROPERTY(EditAnywhere, Category = "Chaos Deformable")
+	UPROPERTY(EditAnywhere, Category = "Physics")
 		bool CacheToFile = false;
 
-	UPROPERTY(EditAnywhere, Category = "Chaos Deformable")
+	UPROPERTY(EditAnywhere, Category = "Physics")
 		bool bEnableKinematics = true;
 
-	UPROPERTY(EditAnywhere, Category = "Chaos Deformable")
+	UPROPERTY(EditAnywhere, Category = "Physics")
 		bool bUseFloor = true;
 
-	UPROPERTY(EditAnywhere, Category = "Chaos Deformable")
+	UPROPERTY(EditAnywhere, Category = "Physics")
 		bool bDoSelfCollision = false;
 
-	UPROPERTY(EditAnywhere, Category = "Chaos Deformable")
+	UPROPERTY(EditAnywhere, Category = "Physics")
 		bool bDoThreadedAdvance = true;
 
-	UPROPERTY(EditAnywhere, Category = "Chaos|Deformable|Expermental" )
+	UPROPERTY(EditAnywhere, Category = "Physics|Expermental" )
 		bool bUseGridBasedConstraints = false;
 
-	UPROPERTY(EditAnywhere, Category = "Chaos Deformable")
+	UPROPERTY(EditAnywhere, Category = "Physics")
 		float GridDx = 25.;
 
-	UPROPERTY(EditAnywhere, Category = "Chaos Deformable")
+	UPROPERTY(EditAnywhere, Category = "Physics")
 		bool bDoQuasistatics = false;
 
-	UPROPERTY(EditAnywhere, Category = "Chaos Deformable")
+	UPROPERTY(EditAnywhere, Category = "Physics")
 		float YoungModulus = 100000;
 
-	UPROPERTY(EditAnywhere, Category = "Chaos Deformable")
+	UPROPERTY(EditAnywhere, Category = "Physics")
 		bool bDoBlended = false;
 
-	UPROPERTY(EditAnywhere, Category = "Chaos Deformable")
+	UPROPERTY(EditAnywhere, Category = "Physics")
 		float BlendedZeta = 0;
 
-	UPROPERTY(EditAnywhere, Category = "Chaos Deformable")
+	UPROPERTY(EditAnywhere, Category = "Physics")
 		float Damping = 0;
 
-	UPROPERTY(EditAnywhere, Category = "Chaos Deformable")
+	UPROPERTY(EditAnywhere, Category = "Physics")
 		bool bEnableGravity = true;
 
-	UPROPERTY(EditAnywhere, Category = "Chaos Deformable")
+	UPROPERTY(EditAnywhere, Category = "Physics")
 		bool bEnableCorotatedConstraint = true;
 
-	UPROPERTY(EditAnywhere, Category = "Chaos Deformable")
+	UPROPERTY(EditAnywhere, Category = "Physics")
 		bool bEnablePositionTargets = true;
 	//UPROPERTY(EditAnywhere, Category = Chaos)
 	//	bool bWaitForParallelDeformableTask = true;
@@ -125,6 +126,10 @@ public:
 	// Simulation Variables
 	TUniquePtr<FDeformableSolver> Solver;
 
+#if WITH_EDITOR
+	virtual void CustomizeDetails(IDetailLayoutBuilder& DetailBuilder) override;
+#endif
+
 protected:
 
 	/** Ref for the deformable solvers parallel task, so we can detect whether or not a sim is running */
diff --git a/Engine/Plugins/Experimental/ChaosFlesh/Source/ChaosFleshEngine/Public/ChaosFlesh/FleshActor.h b/Engine/Plugins/Experimental/ChaosFlesh/Source/ChaosFleshEngine/Public/ChaosFlesh/FleshActor.h
--- a/Engine/Plugins/Experimental/ChaosFlesh/Source/ChaosFleshEngine/Public/ChaosFlesh/FleshActor.h
+++ b/Engine/Plugins/Experimental/ChaosFlesh/Source/ChaosFleshEngine/Public/ChaosFlesh/FleshActor.h
@@ -2,19 +2,20 @@
 
 #pragma once
 
+#include "ChaosFlesh/FleshComponent.h"
+#include "ChaosFlesh/ChaosDeformableSolverActor.h"
 #include "CoreMinimal.h"
+#include "DeformableInterface.h"
 #include "GameFramework/Actor.h"
 #include "UObject/ObjectMacros.h"
 #include "UObject/ObjectPtr.h"
-#include "ChaosFlesh/FleshComponent.h"
-#include "ChaosFlesh/ChaosDeformableSolverActor.h"
 
 #include "FleshActor.generated.h"
 
 
 
 UCLASS()
-class CHAOSFLESHENGINE_API AFleshActor : public AActor
+class CHAOSFLESHENGINE_API AFleshActor : public AActor, public IDeformableInterface
 {
 	GENERATED_UCLASS_BODY()
 
@@ -23,11 +24,11 @@ public:
 	UFUNCTION(BlueprintCallable, Category = "Physics")
 	void EnableSimulation(ADeformableSolverActor* Actor);
 
-	UPROPERTY(VisibleAnywhere, BlueprintReadOnly, Category = "Chaos Deformable")
+	UPROPERTY(VisibleAnywhere, BlueprintReadOnly, Category = "Physics")
 	TObjectPtr<UFleshComponent> FleshComponent;
 	UFleshComponent* GetFleshComponent() const { return FleshComponent; }
 
-	UPROPERTY(EditAnywhere, BlueprintReadWrite, Category = "Chaos Deformable")
+	UPROPERTY(EditAnywhere, BlueprintReadWrite, Category = "Physics")
 	TObjectPtr<ADeformableSolverActor> PrimarySolver;
 
 #if WITH_EDITOR
@@ -35,5 +36,6 @@ public:
 	virtual void PreEditChange(FProperty* PropertyThatWillChange) override;
 	virtual void PostEditChangeProperty(struct FPropertyChangedEvent& PropertyChangedEvent) override;
 	virtual bool GetReferencedContentObjects(TArray<UObject*>& Objects) const override;
+	virtual void CustomizeDetails(IDetailLayoutBuilder& DetailBuilder) override;
 #endif
 };
diff --git a/Engine/Plugins/Experimental/ChaosFlesh/Source/ChaosFleshEngine/Public/ChaosFlesh/FleshComponent.h b/Engine/Plugins/Experimental/ChaosFlesh/Source/ChaosFleshEngine/Public/ChaosFlesh/FleshComponent.h
--- a/Engine/Plugins/Experimental/ChaosFlesh/Source/ChaosFleshEngine/Public/ChaosFlesh/FleshComponent.h
+++ b/Engine/Plugins/Experimental/ChaosFlesh/Source/ChaosFleshEngine/Public/ChaosFlesh/FleshComponent.h
@@ -75,7 +75,7 @@ public:
 	USimulationAsset* GetSimulationCollection() { return SimulationCollection; }
 	const USimulationAsset* GetSimulationCollection() const { return SimulationCollection; }
 
-	UPROPERTY(EditAnywhere, Category = "ChaosPhysics")
+	UPROPERTY(EditAnywhere, Category = "Rendering")
 	TObjectPtr<UProceduralMeshComponent> Mesh;
 
 	/** @deprecated Use GetSkeletalMeshEmbeddedPositions() instead. */
@@ -98,7 +98,7 @@ public:
 
 private:
 	/** FleshAsset that describes the simulation rest state. */
-	UPROPERTY(EditAnywhere, Category = "ChaosPhysics")
+	UPROPERTY(EditAnywhere, Category = "Physics")
 	TObjectPtr<const UFleshAsset> RestCollection;
 
 	/** Current simulation state. */
@@ -114,14 +114,14 @@ private:
 	//
 
 	/** Space the simulation will run in. */
-	UPROPERTY(EditAnywhere, Category = "ChaosDeformable")
+	UPROPERTY(EditAnywhere, Category = "Physics")
 	TEnumAsByte<ChaosDeformableSimSpace> SimSpace = ChaosDeformableSimSpace::World;
 
 	/** 
 	* Bone from the associated skeletal mesh (indicated by RestCollection.TargetSkeletalMesh) to use as 
 	* the space the sim runs in.
 	*/
-	UPROPERTY(EditAnywhere, Category = "ChaosDeformable", meta = (GetOptions = "GetSimSpaceBoneNameOptions", EditCondition = "SimSpace == ChaosDeformableSimSpace::Bone"))
+	UPROPERTY(EditAnywhere, Category = "Physics", meta = (GetOptions = "GetSimSpaceBoneNameOptions", EditCondition = "SimSpace == ChaosDeformableSimSpace::Bone"))
 	FName SimSpaceBoneName;
 
 	/** The skeletal mesh to use pull the \c SimSpaceBoneName from. */
diff --git a/Engine/Source/Runtime/Engine/Public/DeformableInterface.h b/Engine/Source/Runtime/Engine/Public/DeformableInterface.h
--- a/Engine/Source/Runtime/Engine/Public/DeformableInterface.h
+++ b/Engine/Source/Runtime/Engine/Public/DeformableInterface.h
@@ -2,8 +2,14 @@
 
 #pragma once
 #include "UObject/Interface.h"
+
+#if WITH_EDITOR
+#include "DetailLayoutBuilder.h"
+#endif
+
 #include "DeformableInterface.generated.h"
 
+
 UINTERFACE()
 class ENGINE_API UDeformableInterface : public UInterface
 {
@@ -13,4 +19,13 @@ class ENGINE_API UDeformableInterface : public UInterface
 class ENGINE_API IDeformableInterface
 {
 	GENERATED_IINTERFACE_BODY()
+
+public:
+
+
+#if WITH_EDITOR
+	// Take damage
+	virtual void CustomizeDetails(IDetailLayoutBuilder& DetailBuilder) = 0;
+#endif
+
 };
