ChaosFlesh : Update properties and component heirarchy
- Propertu restructure.
- Changed derivation heirarchy for the flesh component.

diff --git a/Engine/Plugins/Experimental/ChaosFlesh/Source/ChaosFleshEngine/Private/ChaosFlesh/ChaosDeformableCollisionsActor.cpp b/Engine/Plugins/Experimental/ChaosFlesh/Source/ChaosFleshEngine/Private/ChaosFlesh/ChaosDeformableCollisionsActor.cpp
--- a/Engine/Plugins/Experimental/ChaosFlesh/Source/ChaosFleshEngine/Private/ChaosFlesh/ChaosDeformableCollisionsActor.cpp
+++ b/Engine/Plugins/Experimental/ChaosFlesh/Source/ChaosFleshEngine/Private/ChaosFlesh/ChaosDeformableCollisionsActor.cpp
@@ -19,8 +19,8 @@ ADeformableCollisionsActor::ADeformableCollisionsActor(const FObjectInitializer&
 #if WITH_EDITOR
 void ADeformableCollisionsActor::CustomizeDetails(IDetailLayoutBuilder& DetailBuilder)
 {
-	//TSharedPtr<IPropertyHandle> bReplicatePhysicsProperty = DetailBuilder.GetProperty(GET_MEMBER_NAME_CHECKED(AFleshActor, bAsyncPhysicsTickEnabled), AActor::StaticClass());
-	//bReplicatePhysicsProperty->MarkHiddenByCustomization();
+	TSharedPtr<IPropertyHandle> bReplicatePhysicsProperty = DetailBuilder.GetProperty(GET_MEMBER_NAME_CHECKED(ADeformableCollisionsActor, bAsyncPhysicsTickEnabled), AActor::StaticClass());
+	bReplicatePhysicsProperty->MarkHiddenByCustomization();
 }
 #endif
 
@@ -86,9 +86,9 @@ void ADeformableCollisionsActor::PostEditChangeProperty(struct FPropertyChangedE
 				if (DeformableCollisionsComponent)
 				{
 					DeformableCollisionsComponent->PrimarySolverComponent = SolverComponent;
-					if (!SolverComponent->DeformableComponents.Contains(DeformableCollisionsComponent))
+					if (!SolverComponent->ConnectedObjects.DeformableComponents.Contains(DeformableCollisionsComponent))
 					{
-						SolverComponent->DeformableComponents.Add(TObjectPtr<UDeformablePhysicsComponent>(DeformableCollisionsComponent));
+						SolverComponent->ConnectedObjects.DeformableComponents.Add(TObjectPtr<UDeformablePhysicsComponent>(DeformableCollisionsComponent));
 					}
 				}
 			}
@@ -100,9 +100,9 @@ void ADeformableCollisionsActor::PostEditChangeProperty(struct FPropertyChangedE
 				if (DeformableCollisionsComponent)
 				{
 					DeformableCollisionsComponent->PrimarySolverComponent = nullptr;
-					if (SolverComponent->DeformableComponents.Contains(DeformableCollisionsComponent))
+					if (SolverComponent->ConnectedObjects.DeformableComponents.Contains(DeformableCollisionsComponent))
 					{
-						SolverComponent->DeformableComponents.Remove(DeformableCollisionsComponent);
+						SolverComponent->ConnectedObjects.DeformableComponents.Remove(DeformableCollisionsComponent);
 					}
 				}
 			}
diff --git a/Engine/Plugins/Experimental/ChaosFlesh/Source/ChaosFleshEngine/Private/ChaosFlesh/ChaosDeformableGameplayComponent.cpp b/Engine/Plugins/Experimental/ChaosFlesh/Source/ChaosFleshEngine/Private/ChaosFlesh/ChaosDeformableGameplayComponent.cpp
--- a/Engine/Plugins/Experimental/ChaosFlesh/Source/ChaosFleshEngine/Private/ChaosFlesh/ChaosDeformableGameplayComponent.cpp
+++ b/Engine/Plugins/Experimental/ChaosFlesh/Source/ChaosFleshEngine/Private/ChaosFlesh/ChaosDeformableGameplayComponent.cpp
@@ -49,9 +49,9 @@ UDeformableGameplayComponent::~UDeformableGameplayComponent()
 void UDeformableGameplayComponent::PreSolverUpdate()
 {
 	PERF_SCOPE(STAT_ChaosDeformable_UDeformableGameplayComponent_PreSolverUpdate)
-	if (RigBoundRayCasts.bEnableRigBoundRaycasts)
+	if (GameplayColllisions.RigBoundRayCasts.bEnableRigBoundRaycasts)
 	{
-		DetectEnvironmentCollisions(RigBoundRayCasts.MaxNumTests, RigBoundRayCasts.bTestDownOnly, RigBoundRayCasts.TestRange, RigBoundRayCasts.CollisionChannel.GetValue());
+		DetectEnvironmentCollisions(GameplayColllisions.RigBoundRayCasts.MaxNumTests, GameplayColllisions.RigBoundRayCasts.bTestDownOnly, GameplayColllisions.RigBoundRayCasts.TestRange, GameplayColllisions.RigBoundRayCasts.CollisionChannel.GetValue());
 	}
 }
 
@@ -316,7 +316,7 @@ void UDeformableGameplayComponent::DetectEnvironmentCollisions(const int32 MaxNu
 		{
 			FHitResult& Hit = HitBuffer[i];
 
-			if (EnvironmentCollisionsSkipList.Contains(Hit.Component))
+			if (GameplayColllisions.RigBoundRayCasts.EnvironmentCollisionsSkipList.Contains(Hit.Component))
 			{
 				continue;
 			}
diff --git a/Engine/Plugins/Experimental/ChaosFlesh/Source/ChaosFleshEngine/Private/ChaosFlesh/ChaosDeformablePhysicsComponent.cpp b/Engine/Plugins/Experimental/ChaosFlesh/Source/ChaosFleshEngine/Private/ChaosFlesh/ChaosDeformablePhysicsComponent.cpp
--- a/Engine/Plugins/Experimental/ChaosFlesh/Source/ChaosFleshEngine/Private/ChaosFlesh/ChaosDeformablePhysicsComponent.cpp
+++ b/Engine/Plugins/Experimental/ChaosFlesh/Source/ChaosFleshEngine/Private/ChaosFlesh/ChaosDeformablePhysicsComponent.cpp
@@ -132,9 +132,9 @@ void UDeformablePhysicsComponent::EnableSimulation(UDeformableSolverComponent* D
 	if (DeformableSolverComponent)
 	{
 		PrimarySolverComponent = DeformableSolverComponent;
-		if (!DeformableSolverComponent->DeformableComponents.Contains(this))
+		if (!DeformableSolverComponent->ConnectedObjects.DeformableComponents.Contains(this))
 		{
-			DeformableSolverComponent->DeformableComponents.Add(this);
+			DeformableSolverComponent->ConnectedObjects.DeformableComponents.Add(this);
 		}
 		DeformableSolverComponent->AddDeformableProxy(this);
 	}
@@ -146,9 +146,9 @@ void UDeformablePhysicsComponent::EnableSimulationFromActor(ADeformableSolverAct
 	if (DeformableSolverActor && DeformableSolverActor->GetDeformableSolverComponent())
 	{
 		PrimarySolverComponent = DeformableSolverActor->GetDeformableSolverComponent();
-		if (!DeformableSolverActor->GetDeformableSolverComponent()->DeformableComponents.Contains(this))
+		if (!DeformableSolverActor->GetDeformableSolverComponent()->ConnectedObjects.DeformableComponents.Contains(this))
 		{
-			DeformableSolverActor->GetDeformableSolverComponent()->DeformableComponents.Add(this);
+			DeformableSolverActor->GetDeformableSolverComponent()->ConnectedObjects.DeformableComponents.Add(this);
 		}
 		DeformableSolverActor->GetDeformableSolverComponent()->AddDeformableProxy(this);
 	}
diff --git a/Engine/Plugins/Experimental/ChaosFlesh/Source/ChaosFleshEngine/Private/ChaosFlesh/ChaosDeformableSolverActor.cpp b/Engine/Plugins/Experimental/ChaosFlesh/Source/ChaosFleshEngine/Private/ChaosFlesh/ChaosDeformableSolverActor.cpp
--- a/Engine/Plugins/Experimental/ChaosFlesh/Source/ChaosFleshEngine/Private/ChaosFlesh/ChaosDeformableSolverActor.cpp
+++ b/Engine/Plugins/Experimental/ChaosFlesh/Source/ChaosFleshEngine/Private/ChaosFlesh/ChaosDeformableSolverActor.cpp
@@ -21,8 +21,11 @@ ADeformableSolverActor::ADeformableSolverActor(const FObjectInitializer& ObjectI
 #if WITH_EDITOR
 void ADeformableSolverActor::CustomizeDetails(IDetailLayoutBuilder& DetailBuilder)
 {
-	//TSharedPtr<IPropertyHandle> bReplicatePhysicsProperty = DetailBuilder.GetProperty(GET_MEMBER_NAME_CHECKED(AFleshActor, bAsyncPhysicsTickEnabled), AActor::StaticClass());
-	//bReplicatePhysicsProperty->MarkHiddenByCustomization();
+	TSharedPtr<IPropertyHandle> bAsyncPhysicsTickEnabledProperty = DetailBuilder.GetProperty(GET_MEMBER_NAME_CHECKED(ADeformableSolverActor, bAsyncPhysicsTickEnabled), AActor::StaticClass());
+	bAsyncPhysicsTickEnabledProperty->MarkHiddenByCustomization();
+
+	TSharedPtr<IPropertyHandle> ShouldUpdatePhysicsVolumnProperty = DetailBuilder.GetProperty("bShouldUpdatePhysicsVolume", USceneComponent::StaticClass());
+	ShouldUpdatePhysicsVolumnProperty->MarkHiddenByCustomization();
 }
 
 #endif
diff --git a/Engine/Plugins/Experimental/ChaosFlesh/Source/ChaosFleshEngine/Private/ChaosFlesh/ChaosDeformableSolverComponent.cpp b/Engine/Plugins/Experimental/ChaosFlesh/Source/ChaosFleshEngine/Private/ChaosFlesh/ChaosDeformableSolverComponent.cpp
--- a/Engine/Plugins/Experimental/ChaosFlesh/Source/ChaosFleshEngine/Private/ChaosFlesh/ChaosDeformableSolverComponent.cpp
+++ b/Engine/Plugins/Experimental/ChaosFlesh/Source/ChaosFleshEngine/Private/ChaosFlesh/ChaosDeformableSolverComponent.cpp
@@ -43,8 +43,8 @@ UDeformableSolverComponent::~UDeformableSolverComponent()
 #if WITH_EDITOR
 void UDeformableSolverComponent::CustomizeDetails(IDetailLayoutBuilder& DetailBuilder)
 {
-	//TSharedPtr<IPropertyHandle> bReplicatePhysicsProperty = DetailBuilder.GetProperty(GET_MEMBER_NAME_CHECKED(AFleshActor, bAsyncPhysicsTickEnabled), AActor::StaticClass());
-	//bReplicatePhysicsProperty->MarkHiddenByCustomization();
+	TSharedPtr<IPropertyHandle> ShouldUpdatePhysicsVolumnProperty = DetailBuilder.GetProperty("bShouldUpdatePhysicsVolume", USceneComponent::StaticClass());
+	ShouldUpdatePhysicsVolumnProperty->MarkHiddenByCustomization();
 }
 #endif
 
@@ -55,12 +55,12 @@ void UDeformableSolverComponent::UpdateTickGroup()
 	// see : CacheManagerActor.cpp::348
 	
 	//
-	if (ExecutionModel == EDeformableExecutionModel::Chaos_Deformable_PrePhysics)
+	if (SolverTiming.ExecutionModel == EDeformableExecutionModel::Chaos_Deformable_PrePhysics)
 	{
 		PrimaryComponentTick.TickGroup = TG_PrePhysics;
 		DeformableEndTickFunction.TickGroup = TG_PrePhysics;
 	}
-	else if (ExecutionModel == EDeformableExecutionModel::Chaos_Deformable_PostPhysics)
+	else if (SolverTiming.ExecutionModel == EDeformableExecutionModel::Chaos_Deformable_PostPhysics)
 	{
 		PrimaryComponentTick.TickGroup = TG_PostPhysics;
 		DeformableEndTickFunction.TickGroup = TG_LastDemotable;
@@ -74,7 +74,7 @@ void UDeformableSolverComponent::UpdateTickGroup()
 	PrimaryComponentTick.bCanEverTick = true;
 	PrimaryComponentTick.bTickEvenWhenPaused = false;
 
-	if (bDoThreadedAdvance)
+	if (SolverTiming.bDoThreadedAdvance)
 	{
 		DeformableEndTickFunction.bCanEverTick = true;
 		DeformableEndTickFunction.bStartWithTickEnabled = true;
@@ -116,7 +116,7 @@ void UDeformableSolverComponent::UpdateDeformableEndTickState(bool bRegister)
 	TRACE_CPUPROFILER_EVENT_SCOPE(ChaosDeformable_UDeformableSolverComponent_UpdateDeformableEndTickState);
 
 	bRegister &= PrimaryComponentTick.IsTickFunctionRegistered();
-	if (bDoThreadedAdvance)
+	if (SolverTiming.bDoThreadedAdvance)
 	{
 		if (bRegister != DeformableEndTickFunction.IsTickFunctionRegistered())
 		{
@@ -165,7 +165,7 @@ void UDeformableSolverComponent::TickComponent(float DeltaTime, enum ELevelTick
 
 		UpdateFromGameThread(DeltaTime);
 
-		if (bDoThreadedAdvance)
+		if (SolverTiming.bDoThreadedAdvance)
 		{
 			// see FParallelClothCompletionTask
 			FGraphEventArray Prerequisites;
@@ -195,27 +195,27 @@ void UDeformableSolverComponent::Reset()
 	if (GChaosEngineDeformableCVarParams.bEnableDeformableSolver)
 	{
 		Solver.Reset(new FDeformableSolver({
-			NumSubSteps
-			, NumSolverIterations
-			, FixTimeStep
-			, TimeStepSize
-			, CacheToFile
-			, bEnableKinematics
-			, bUseFloor
-			, bDoSelfCollision
-			, bUseGridBasedConstraints
-			, GridDx
-			, bDoQuasistatics
-			, YoungModulus
-			, bDoBlended
-			, BlendedZeta
-			, Damping
-			, bEnableGravity
-			, bEnableCorotatedConstraint
-			, bEnablePositionTargets
+			SolverTiming.NumSubSteps
+			, SolverTiming.NumSolverIterations
+			, SolverTiming.FixTimeStep
+			, SolverTiming.TimeStepSize
+			, SolverDebugging.CacheToFile
+			, SolverConstraints.bEnableKinematics
+			, SolverCollisions.bUseFloor
+			, SolverCollisions.bDoSelfCollision
+			, false /*SolverCollisions.SolverGridBasedCollisions.bUseGridBasedConstraints*/
+			, 25. /*SolverCollisions.SolverGridBasedCollisions.GridDx*/
+			, SolverEvolution.SolverQuasistatics.bDoQuasistatics
+			, SolverForces.YoungModulus
+			, SolverConstraints.CorotatedConstraints.bDoBlended
+			, SolverConstraints.CorotatedConstraints.BlendedZeta
+			, SolverForces.Damping
+			, SolverForces.bEnableGravity
+			, SolverConstraints.CorotatedConstraints.bEnableCorotatedConstraint
+			, SolverConstraints.bEnablePositionTargets
 		}));
 
-		for (TObjectPtr<UDeformablePhysicsComponent>& DeformableComponent : DeformableComponents)
+		for (TObjectPtr<UDeformablePhysicsComponent>& DeformableComponent : ConnectedObjects.DeformableComponents)
 		{
 			if( DeformableComponent )
 			{
@@ -266,7 +266,7 @@ void UDeformableSolverComponent::UpdateFromGameThread(float DeltaTime)
 		FDeformableSolver::FGameThreadAccess GameThreadSolver(Solver.Get(), Chaos::Softs::FGameThreadAccessor());
 
 		Chaos::Softs::FDeformableDataMap DataMap;
-		for (TObjectPtr<UDeformablePhysicsComponent>& DeformableComponent : DeformableComponents)
+		for (TObjectPtr<UDeformablePhysicsComponent>& DeformableComponent : ConnectedObjects.DeformableComponents)
 		{
 			if (DeformableComponent)
 			{
@@ -301,7 +301,7 @@ void UDeformableSolverComponent::UpdateFromSimulation(float DeltaTime)
 
 		if (Output)
 		{
-			for (TObjectPtr<UDeformablePhysicsComponent>& DeformableComponent : DeformableComponents)
+			for (TObjectPtr<UDeformablePhysicsComponent>& DeformableComponent : ConnectedObjects.DeformableComponents)
 			{
 				if (DeformableComponent)
 				{
diff --git a/Engine/Plugins/Experimental/ChaosFlesh/Source/ChaosFleshEngine/Private/ChaosFlesh/ChaosDeformableTetrahedralComponent.cpp b/Engine/Plugins/Experimental/ChaosFlesh/Source/ChaosFleshEngine/Private/ChaosFlesh/ChaosDeformableTetrahedralComponent.cpp
new file mode 100644
--- /dev/null
+++ b/Engine/Plugins/Experimental/ChaosFlesh/Source/ChaosFleshEngine/Private/ChaosFlesh/ChaosDeformableTetrahedralComponent.cpp
@@ -0,0 +1,943 @@
+// Copyright Epic Games, Inc. All Rights Reserved.
+
+#include "ChaosFlesh/ChaosDeformableTetrahedralComponent.h"
+
+#include "Animation/SkeletalMeshActor.h"
+#include "Animation/Skeleton.h"
+#include "ChaosStats.h"
+#include "Chaos/DebugDrawQueue.h"
+#include "Chaos/Tetrahedron.h"
+#include "ChaosFlesh/FleshCollection.h"
+#include "ChaosFlesh/ChaosDeformableSolverActor.h"
+#include "ChaosFlesh/ChaosDeformableSolverComponent.h"
+#include "ChaosFlesh/ChaosDeformableTypes.h"
+#include "ChaosFlesh/FleshDynamicAsset.h"
+#include "Components/SkeletalMeshComponent.h"
+#include "Dataflow/DataflowEngineUtil.h"
+#include "Engine/SkeletalMesh.h"
+#include "GeometryCollection/Facades/CollectionTransformSourceFacade.h"
+#include "GeometryCollection/Facades/CollectionTetrahedralSkeletalBindingsFacade.h"
+#include "GeometryCollection/TransformCollection.h"
+#include "ProceduralMeshComponent.h"
+#include "Rendering/SkeletalMeshRenderData.h"
+
+
+#include UE_INLINE_GENERATED_CPP_BY_NAME(ChaosDeformableTetrahedralComponent)
+
+FChaosEngineDeformableCVarParams CVarParams;
+FAutoConsoleVariableRef CVarDeforambleDoDrawSimulationMesh(TEXT("p.Chaos.DebugDraw.Deformable.SimulationMesh"), CVarParams.bDoDrawSimulationMesh, TEXT("Debug draw the deformable simulation resutls on the game thread. [def: true]"));
+FAutoConsoleVariableRef CVarDeforambleDoDrawSkeletalMeshBindingPositions(TEXT("p.Chaos.DebugDraw.Deformable.SkeletalMeshBindingPositions"), CVarParams.bDoDrawSkeletalMeshBindingPositions, TEXT("Debug draw the deformable simulation's SkeletalMeshBindingPositions on the game thread. [def: false]"));
+FAutoConsoleVariableRef CVarDeforambleDoDrawSkeletalMeshBindingPositionsSimulationBlendWeight(TEXT("p.Chaos.DebugDraw.Deformable.SkeletalMeshBindingPositions.SimulationBlendWeight"), CVarParams.DrawSkeletalMeshBindingPositionsSimulationBlendWeight, TEXT("Set the simulation blend weight of the skeletal mesh debug draw.[def: 1.]"));
+
+
+#define PERF_SCOPE(X) SCOPE_CYCLE_COUNTER(X); TRACE_CPUPROFILER_EVENT_SCOPE(X);
+DECLARE_CYCLE_STAT(TEXT("Chaos.Deformable.UDeformableTetrahedralComponent.TickComponent"), STAT_ChaosDeformable_UDeformableTetrahedralComponent_TickComponent, STATGROUP_Chaos);
+DECLARE_CYCLE_STAT(TEXT("Chaos.Deformable.UDeformableTetrahedralComponent.NewProxy"), STAT_ChaosDeformable_UDeformableTetrahedralComponent_NewProxy, STATGROUP_Chaos);
+DECLARE_CYCLE_STAT(TEXT("Chaos.Deformable.UDeformableTetrahedralComponent.NewDeformableData"), STAT_ChaosDeformable_UDeformableTetrahedralComponent_NewDeformableData, STATGROUP_Chaos);
+DECLARE_CYCLE_STAT(TEXT("Chaos.Deformable.UDeformableTetrahedralComponent.UpdateFromSimualtion"), STAT_ChaosDeformable_UDeformableTetrahedralComponent_UpdateFromSimualtion, STATGROUP_Chaos);
+DECLARE_CYCLE_STAT(TEXT("Chaos.Deformable.UDeformableTetrahedralComponent.RenderProceduralMesh"), STAT_ChaosDeformable_UDeformableTetrahedralComponent_RenderProceduralMesh, STATGROUP_Chaos);
+
+DEFINE_LOG_CATEGORY_STATIC(LogDeformableTetrahedralComponentInternal, Log, All);
+
+UDeformableTetrahedralComponent::UDeformableTetrahedralComponent(const FObjectInitializer& ObjectInitializer)
+	: Super(ObjectInitializer)
+{
+	Mesh = ObjectInitializer.CreateDefaultSubobject<UProceduralMeshComponent>(this, TEXT("Flesh Visualization Component"));
+	PrimaryComponentTick.TickGroup = TG_LastDemotable;
+	PrimaryComponentTick.bCanEverTick = CVarParams.bDoDrawSimulationMesh;
+	bTickInEditor = CVarParams.bDoDrawSimulationMesh;
+
+	DynamicCollection = ObjectInitializer.CreateDefaultSubobject<UFleshDynamicAsset>(this, TEXT("Flesh Dynamic Asset"));
+	SimulationCollection = ObjectInitializer.CreateDefaultSubobject<USimulationAsset>(this, TEXT("Flesh Simulation Asset"));
+}
+
+UDeformableTetrahedralComponent::~UDeformableTetrahedralComponent()
+{
+	if (RenderMesh) delete RenderMesh;
+}
+
+void UDeformableTetrahedralComponent::Invalidate()
+{
+	bBoundsNeedsUpdate = true;
+}
+
+
+void UDeformableTetrahedralComponent::OnRegister()
+{
+	if (bBoundsNeedsUpdate)
+	{
+		UpdateLocalBounds();
+	}
+	Super::OnRegister();
+}
+
+void UDeformableTetrahedralComponent::EndPlay(const EEndPlayReason::Type ReasonEnd)
+{
+	if (GetDynamicCollection())
+	{
+		GetDynamicCollection()->Reset();
+	}
+	if (GetSimulationCollection())
+	{
+		GetSimulationCollection()->Reset();
+	}
+}
+
+void UDeformableTetrahedralComponent::SetRestCollection(const UFleshAsset* InRestCollection)
+{
+	RestCollection = InRestCollection;
+	Invalidate();
+	UpdateLocalBounds();
+	ResetProceduralMesh();
+}
+
+
+UDeformablePhysicsComponent::FThreadingProxy* UDeformableTetrahedralComponent::NewProxy()
+{
+	PERF_SCOPE(STAT_ChaosDeformable_UDeformableTetrahedralComponent_NewProxy);
+
+	UpdateSimSpaceTransformIndex();
+	if (const UFleshAsset* RestAsset = GetRestCollection())
+	{
+		if (const FFleshCollection* Rest = RestAsset->GetCollection())
+		{
+			if (Rest->NumElements(FGeometryCollection::VerticesGroup))
+			{
+				if (!GetDynamicCollection())
+				{
+					DynamicCollection = NewObject<UFleshDynamicAsset>(this, TEXT("Flesh Dynamic Asset"));
+				}
+				if (!GetSimulationCollection())
+				{
+					SimulationCollection = NewObject<USimulationAsset>(this, TEXT("Flesh Simulation Asset"));
+				}
+
+				GetDynamicCollection()->Reset(Rest);
+				if (const FManagedArrayCollection* Dynamic = GetDynamicCollection()->GetCollection())
+				{
+					// Mesh points are in component space, such that the exterior hull aligns with the
+					// surface of the skeletal mesh, which is subject to the transform hierarchy.
+					const FTransform& ComponentToWorldXf = GetComponentTransform();
+					const FTransform ComponentToSimXf = GetSimSpaceRestTransform();
+					return new FFleshThreadingProxy(
+						this,
+						ComponentToWorldXf,
+						ComponentToSimXf,
+						SimulationSpace.SimSpace,
+						*Rest,
+						*Dynamic);
+				}
+			}
+		}
+	}
+	return nullptr;
+}
+
+UDeformablePhysicsComponent::FDataMapValue UDeformableTetrahedralComponent::NewDeformableData()
+{
+	PERF_SCOPE(STAT_ChaosDeformable_UDeformableTetrahedralComponent_NewDeformableData);
+
+	using namespace GeometryCollection::Facades;
+	if (GetOwner())
+	{
+		if (const UFleshAsset* FleshAsset = GetRestCollection())
+		{
+			if (const FFleshCollection* Rest = FleshAsset->GetCollection())
+			{
+				FTransformSource TransformSource(*Rest);
+				if (TransformSource.IsValid())
+				{
+					TArray<USkeletalMeshComponent*> SkeletalMeshComponents;
+					GetOwner()->GetComponents<USkeletalMeshComponent>(SkeletalMeshComponents);
+
+					if (const TManagedArray<FTransform>* RestTransforms = Rest->FindAttribute<FTransform>(FTransformCollection::TransformAttribute, FTransformCollection::TransformGroup))
+					{
+						TArray<FTransform> AnimationTransforms = RestTransforms->GetConstArray();
+						TArray<FTransform> ComponentPose = RestTransforms->GetConstArray();
+
+						// Extract animated transforms from all skeletal meshes.
+						for (const USkeletalMeshComponent* SkeletalMeshComponent : SkeletalMeshComponents)
+						{
+							if (const USkeletalMesh* SkeletalMesh = SkeletalMeshComponent->GetSkeletalMeshAsset())
+							{
+								if (const USkeleton* Skeleton = SkeletalMesh->GetSkeleton())
+								{
+									TSet<int32> Roots = TransformSource.GetTransformSource(Skeleton->GetName(), Skeleton->GetGuid().ToString());
+									if (!Roots.IsEmpty() && ensureMsgf(Roots.Num() == 1, TEXT("Error: Only supports a single root per skeleton.(%s)"), *Skeleton->GetName()))
+									{
+										TArray<FTransform> ComponentLocalPose;
+										Dataflow::Animation::GlobalTransforms(SkeletalMesh->GetRefSkeleton(), ComponentLocalPose);
+
+										const TArray<FTransform>& ComponentTransforms = SkeletalMeshComponent->GetComponentSpaceTransforms();
+										if (ComponentLocalPose.Num() == ComponentTransforms.Num())
+										{
+											for (int Adx = Roots.Array()[0], Cdx = 0; Adx <= AnimationTransforms.Num() && Cdx < ComponentTransforms.Num(); Adx++, Cdx++)
+											{
+												// @todo(flesh) : Can we just use one array?
+												AnimationTransforms[Adx] = ComponentTransforms[Cdx];
+												ComponentPose[Adx] = ComponentLocalPose[Cdx]; 
+
+												if (SimulationSpace.SimSpaceTransformGlobalIndex == INDEX_NONE &&
+													SimulationSpace.SimSpaceTransformIndex == Cdx &&
+													SkeletalMesh == SimulationSpace.SimSpaceSkeletalMesh)
+												{
+													SimulationSpace.SimSpaceTransformGlobalIndex = Adx;
+												}
+											}
+
+
+										}
+									}
+								}
+							}
+						}
+
+						FTransform BoneSpaceXf;
+						if (AnimationTransforms.IsValidIndex(SimulationSpace.SimSpaceTransformGlobalIndex))
+						{
+							BoneSpaceXf = AnimationTransforms[SimulationSpace.SimSpaceTransformGlobalIndex];
+						}
+						else
+						{
+							BoneSpaceXf = FTransform::Identity;
+						}
+
+						return FDataMapValue(
+							new Chaos::Softs::FFleshThreadingProxy::FFleshInputBuffer(
+								*GetSimulationCollection()->GetCollection(),
+								this->GetComponentTransform(),
+								BoneSpaceXf, 
+							SimulationSpace.SimSpaceTransformGlobalIndex,
+								AnimationTransforms, 
+								ComponentPose, 
+								BodyForces.bApplyGravity,
+								BodyForces.StiffnessMultiplier,
+								BodyForces.DampingMultiplier,
+								MassMultiplier,
+								BodyForces.IncompressibilityMultiplier,
+								BodyForces.InflationMultiplier,
+								this));
+					}
+				}
+			}
+		}
+	}
+	return FDataMapValue(
+		new Chaos::Softs::FFleshThreadingProxy::FFleshInputBuffer(
+			*GetSimulationCollection()->GetCollection(),
+			this->GetComponentTransform(),
+			GetSimSpaceRestTransform(),
+		SimulationSpace.SimSpaceTransformGlobalIndex,
+			BodyForces.bApplyGravity,
+			BodyForces.StiffnessMultiplier,
+			BodyForces.DampingMultiplier,
+			MassMultiplier,
+			BodyForces.IncompressibilityMultiplier,
+			BodyForces.InflationMultiplier,
+			this));
+}
+
+TArray<FString> UDeformableTetrahedralComponent::GetSimSpaceBoneNameOptions() const
+{
+	TArray<FString> Names;
+	if (RestCollection)
+	{
+		if (RestCollection->SkeletalMesh)
+		{
+			const FReferenceSkeleton& RefSkeleton = RestCollection->SkeletalMesh->GetSkeleton()->GetReferenceSkeleton();
+			Names.SetNum(RefSkeleton.GetNum());
+			for (int32 i = 0; i < RefSkeleton.GetNum(); i++)
+			{
+				Names[i] = RefSkeleton.GetBoneName(i).ToString();
+			}
+		}
+	}
+	return Names;
+}
+
+bool UDeformableTetrahedralComponent::UpdateSimSpaceTransformIndex()
+{
+	SimulationSpace.SimSpaceTransformIndex = INDEX_NONE;
+	SimulationSpace.SimSpaceSkeletalMesh = nullptr;
+
+	if (SimulationSpace.SimSpace != ChaosDeformableSimSpace::Bone)
+	{
+		return false;
+	}
+
+	if (RestCollection && RestCollection->SkeletalMesh)
+	{
+		const FReferenceSkeleton& RefSkeleton = RestCollection->SkeletalMesh->GetSkeleton()->GetReferenceSkeleton();
+		for (int32 i = 0; i < RefSkeleton.GetNum(); i++)
+		{
+			if (RefSkeleton.GetBoneName(i).ToString() == SimulationSpace.SimSpaceBoneName.ToString())
+			{
+				SimulationSpace.SimSpaceSkeletalMesh = RestCollection->SkeletalMesh;
+				SimulationSpace.SimSpaceTransformIndex = i;
+				return true;
+			}
+		}
+	}
+	return false;
+}
+
+FTransform UDeformableTetrahedralComponent::GetSimSpaceRestTransform() const
+{
+	if (SimulationSpace.SimSpaceSkeletalMesh == nullptr)
+	{
+		return FTransform::Identity;
+	}
+	
+	TArray<FTransform> ComponentTransforms;
+	ComponentTransforms.SetNum(SimulationSpace.SimSpaceSkeletalMesh->GetRefSkeleton().GetNum());
+
+	SimulationSpace.SimSpaceSkeletalMesh->FillComponentSpaceTransforms(
+		SimulationSpace.SimSpaceSkeletalMesh->GetRefSkeleton().GetRefBonePose(),
+		SimulationSpace.SimSpaceSkeletalMesh->GetResourceForRendering()->LODRenderData[0].RequiredBones,
+		ComponentTransforms);
+
+	if (!ComponentTransforms.IsValidIndex(SimulationSpace.SimSpaceTransformIndex))
+	{
+		return FTransform::Identity;
+	}
+	const FTransform& ComponentToBone = ComponentTransforms[SimulationSpace.SimSpaceTransformIndex];
+	return ComponentToBone;
+}
+
+void UDeformableTetrahedralComponent::UpdateFromSimualtion(const FDataMapValue* SimualtionBuffer)
+{
+	PERF_SCOPE(STAT_ChaosDeformable_UDeformableTetrahedralComponent_UpdateFromSimualtion);
+
+	if (const FFleshThreadingProxy::FFleshOutputBuffer* FleshBuffer = (*SimualtionBuffer)->As<FFleshThreadingProxy::FFleshOutputBuffer>())
+	{
+		if (GetDynamicCollection())
+		{
+			// @todo(flesh) : reduce conversions
+			auto UEVertd = [](FVector3f V) { return FVector3d(V.X, V.Y, V.Z); };
+			auto UEVertf = [](FVector3d V) { return FVector3f(V.X, V.Y, V.Z); };
+
+			TManagedArray<FVector3f>& DynamicVertex = GetDynamicCollection()->GetPositions();
+			const TManagedArray<FVector3f>& SimulationVertex = FleshBuffer->Dynamic.GetAttribute<FVector3f>("Vertex", FGeometryCollection::VerticesGroup);
+
+			// Simulator produces results in component space.
+			for (int i = DynamicVertex.Num() - 1; i >= 0; i--)
+			{
+				DynamicVertex[i] = SimulationVertex[i];
+			}
+			
+			//p.Chaos.DebugDraw.Enabled 1
+			//p.Chaos.DebugDraw.Deformable.SkeletalMeshBindingPositions 1
+			if (CVarParams.bDoDrawSkeletalMeshBindingPositions)
+			{
+				DebugDrawSkeletalMeshBindingPositions();
+			}
+		}
+	}
+}
+
+void UDeformableTetrahedralComponent::UpdateLocalBounds()
+{
+	if (bBoundsNeedsUpdate && RestCollection)
+	{
+		{
+			FFleshAssetEdit EditObject = RestCollection->EditCollection();
+			if (FFleshCollection* Collection = EditObject.GetFleshCollection())
+			{
+				Collection->UpdateBoundingBox();
+			}
+		}
+		BoundingBox = RestCollection->GetCollection()->GetBoundingBox();
+		bBoundsNeedsUpdate = false;
+	}
+
+}
+
+FBoxSphereBounds UDeformableTetrahedralComponent::CalcBounds(const FTransform& LocalToWorldIn) const
+{
+	return BoundingBox.TransformBy(GetComponentTransform()); // todo(chaos:flesh) use LocalToWorldIn
+}
+
+void UDeformableTetrahedralComponent::ResetDynamicCollection()
+{
+	if (const UFleshAsset* RestAsset = GetRestCollection())
+	{
+		if (!GetDynamicCollection())
+		{
+			DynamicCollection = NewObject<UFleshDynamicAsset>(this, TEXT("Flesh Dynamic Asset"));
+		}
+
+		if (!GetDynamicCollection()->GetCollection() || 
+			!GetDynamicCollection()->GetCollection()->NumElements(FGeometryCollection::VerticesGroup))
+		{
+			GetDynamicCollection()->Reset(RestAsset->GetCollection());
+			ResetProceduralMesh();
+		}
+		else
+		{
+			GetDynamicCollection()->ResetAttributesFrom(RestAsset->GetCollection());
+		}
+	}
+}
+
+//
+// Rendering Support
+//
+
+void UDeformableTetrahedralComponent::BeginPlay()
+{
+	Super::BeginPlay();
+
+	if (ASkeletalMeshActor* SkeletalMeshActor = Cast<ASkeletalMeshActor>(GetOwner()))
+	{
+		if (USkeletalMeshComponent* SkeletalMeshComponent = SkeletalMeshActor->GetSkeletalMeshComponent())
+		{
+			PrimaryComponentTick.AddPrerequisite(SkeletalMeshComponent, SkeletalMeshComponent->PrimaryComponentTick);
+		}
+	}
+	if (PrimarySolverComponent)
+	{
+		PrimaryComponentTick.AddPrerequisite(PrimarySolverComponent, PrimarySolverComponent->PrimaryComponentTick);
+	}
+}
+
+void UDeformableTetrahedralComponent::TickComponent(float DeltaTime, enum ELevelTick TickType, FActorComponentTickFunction* ThisTickFunction)
+{
+	PERF_SCOPE(STAT_ChaosDeformable_UDeformableTetrahedralComponent_TickComponent);
+	Super::TickComponent(DeltaTime, TickType, ThisTickFunction);
+
+	if (IsVisible())
+	{
+		RenderProceduralMesh();
+	}
+	else
+	{
+		ResetProceduralMesh();
+	}
+}
+
+void UDeformableTetrahedralComponent::ResetProceduralMesh()
+{
+	if (Mesh && RenderMesh)
+	{
+		Mesh->ClearAllMeshSections();
+	}
+	if (RenderMesh)
+	{
+		delete RenderMesh;
+		RenderMesh = nullptr;
+	}
+}
+
+void UDeformableTetrahedralComponent::RenderProceduralMesh()
+{
+	bool bCanRender = false;
+	if (const UFleshAsset* FleshAsset = GetRestCollection())
+	{
+		if (IsVisible())
+		{
+#if WITH_EDITORONLY_DATA
+			if (FleshAsset->bRenderInEditor)
+#endif
+			{
+				if (Mesh && CVarParams.bDoDrawSimulationMesh)
+				{
+					PERF_SCOPE(STAT_ChaosDeformable_UDeformableTetrahedralComponent_RenderProceduralMesh);
+
+					if (const FFleshCollection* Flesh = FleshAsset->GetCollection())
+					{
+						int32 NumVertices = Flesh->NumElements(FGeometryCollection::VerticesGroup);
+						int32 NumFaces = Flesh->NumElements(FGeometryCollection::FacesGroup);
+						if (NumFaces && NumVertices)
+						{
+							if (RenderMesh && RenderMesh->Vertices.Num() != NumFaces * 3)
+							{
+								ResetProceduralMesh();
+							}
+
+							if (!RenderMesh)
+							{
+								RenderMesh = new FFleshRenderMesh;
+
+								for (int i = 0; i < NumFaces; ++i)
+								{
+									const auto& P1 = Flesh->Vertex[Flesh->Indices[i][0]];
+									const auto& P2 = Flesh->Vertex[Flesh->Indices[i][1]];
+									const auto& P3 = Flesh->Vertex[Flesh->Indices[i][2]];
+
+									RenderMesh->Vertices.Add(FVector(P1));
+									RenderMesh->Vertices.Add(FVector(P2));
+									RenderMesh->Vertices.Add(FVector(P3));
+
+									RenderMesh->Colors.Add(FLinearColor::White);
+									RenderMesh->Colors.Add(FLinearColor::White);
+									RenderMesh->Colors.Add(FLinearColor::White);
+
+									RenderMesh->UVs.Add(FVector2D(0, 0));
+									RenderMesh->UVs.Add(FVector2D(0, 0));
+									RenderMesh->UVs.Add(FVector2D(0, 0));
+
+									RenderMesh->Triangles.Add(3 * i);
+									RenderMesh->Triangles.Add(3 * i + 1);
+									RenderMesh->Triangles.Add(3 * i + 2);
+
+									auto Normal = -Chaos::FVec3::CrossProduct(P3 - P1, P2 - P1);
+									RenderMesh->Normals.Add(Normal);
+									RenderMesh->Normals.Add(Normal);
+									RenderMesh->Normals.Add(Normal);
+
+									auto Tangent = (P2 - P1).GetSafeNormal();
+									RenderMesh->Tangents.Add(FProcMeshTangent(Tangent[0], Tangent[1], Tangent[2]));
+									Tangent = (P3 - P2).GetSafeNormal();
+									RenderMesh->Tangents.Add(FProcMeshTangent(Tangent[0], Tangent[1], Tangent[2]));
+									Tangent = (P1 - P3).GetSafeNormal();
+									RenderMesh->Tangents.Add(FProcMeshTangent(Tangent[0], Tangent[1], Tangent[2]));
+								}
+
+								Mesh->SetRelativeTransform(GetComponentTransform());
+								Mesh->CreateMeshSection_LinearColor(0, RenderMesh->Vertices, RenderMesh->Triangles, RenderMesh->Normals, RenderMesh->UVs, RenderMesh->Colors, RenderMesh->Tangents, false);
+							}
+							else
+							{
+
+								const TManagedArray<FVector3f>* RenderVertex = &Flesh->Vertex;
+								if (GetDynamicCollection())
+								{
+									const TManagedArray<FVector3f>& DynamicVertex = GetDynamicCollection()->GetPositions();
+									if (DynamicVertex.Num()) RenderVertex = &DynamicVertex;
+								}
+								auto InRange = [](int32 Size, int32 Val) { return 0 <= Val && Val < Size; };
+
+								// Display only
+								for (int i = 0; i < NumFaces; ++i)
+								{
+									const auto& P1 = (*RenderVertex)[Flesh->Indices[i][0]];
+									const auto& P2 = (*RenderVertex)[Flesh->Indices[i][1]];
+									const auto& P3 = (*RenderVertex)[Flesh->Indices[i][2]];
+
+									RenderMesh->Vertices[3 * i] = FVector(P1);
+									RenderMesh->Vertices[3 * i + 1] = FVector(P2);
+									RenderMesh->Vertices[3 * i + 2] = FVector(P3);
+
+									auto Normal = Chaos::FVec3::CrossProduct(P3 - P1, P2 - P1);
+									RenderMesh->Normals[3 * i] = Normal;
+									RenderMesh->Normals[3 * i + 1] = Normal;
+									RenderMesh->Normals[3 * i + 2] = Normal;
+
+									auto Tangent = (P2 - P1).GetSafeNormal();
+									RenderMesh->Tangents[3 * i] = FProcMeshTangent(Tangent[0], Tangent[1], Tangent[2]);
+									Tangent = (P3 - P2).GetSafeNormal();
+									RenderMesh->Tangents[3 * i + 1] = FProcMeshTangent(Tangent[0], Tangent[1], Tangent[2]);
+									Tangent = (P1 - P3).GetSafeNormal();
+									RenderMesh->Tangents[3 * i + 2] = FProcMeshTangent(Tangent[0], Tangent[1], Tangent[2]);
+								}
+
+								if (!Mesh->GetComponentTransform().Equals(GetComponentTransform())) {
+									Mesh->SetRelativeTransform(GetComponentTransform());
+								}
+								Mesh->UpdateMeshSection_LinearColor(0, RenderMesh->Vertices, RenderMesh->Normals, RenderMesh->UVs, RenderMesh->Colors, RenderMesh->Tangents);
+							}
+
+							bCanRender = true;
+						}
+					}
+				}
+			}
+		}
+	}
+	if (!bCanRender)
+	{
+		ResetProceduralMesh();
+	}
+}
+
+
+void UDeformableTetrahedralComponent::DebugDrawSkeletalMeshBindingPositions() const
+{
+#if WITH_EDITOR
+	auto UEVertd = [](FVector3f V) { return FVector3d(V.X, V.Y, V.Z); };
+	float SimulationBlendWeight = CVarParams.DrawSkeletalMeshBindingPositionsSimulationBlendWeight;
+
+	if (const UFleshAsset* RestAsset = GetRestCollection())
+	{
+		const USkeletalMesh* SkeletalMesh = 
+			GetRestCollection() && GetRestCollection()->TargetDeformationSkeleton ? 
+			GetRestCollection()->TargetDeformationSkeleton : RestAsset->SkeletalMesh;
+		if (SkeletalMesh)
+		{
+			TArray<bool> Influenced;
+			TArray<FVector> PosArray = GetSkeletalMeshEmbeddedPositionsInternal(
+				ChaosDeformableBindingOption::ComponentPos, FTransform::Identity, "", SimulationBlendWeight, &Influenced);
+
+			for (int i = 0; i < PosArray.Num(); i++)
+			{
+				const FVector& Pos = PosArray[i];
+				if (Influenced[i])
+				{
+					Chaos::FDebugDrawQueue::GetInstance().DrawDebugPoint(
+						GetComponentTransform().TransformPosition(Pos), FColor::Yellow, true, 2.0f, SDPG_Foreground, 10);
+				}
+			}
+		}
+	}
+#endif
+}
+
+TArray<FVector> UDeformableTetrahedralComponent::GetSkeletalMeshBindingPositions(const USkeletalMesh* InSkeletalMesh) const
+{
+	return GetSkeletalMeshBindingPositionsInternal(InSkeletalMesh, nullptr);
+}
+
+TArray<FVector> UDeformableTetrahedralComponent::GetSkeletalMeshEmbeddedPositions(
+	const ChaosDeformableBindingOption Format,
+	const FTransform TargetDeformationSkeletonOffset,
+	const FName TargetBone,
+	const float SimulationBlendWeight) const
+{
+	return GetSkeletalMeshEmbeddedPositionsInternal(Format, TargetDeformationSkeletonOffset, TargetBone, SimulationBlendWeight, nullptr);
+}
+
+TArray<FVector> UDeformableTetrahedralComponent::GetSkeletalMeshEmbeddedPositionsInternal(
+	const ChaosDeformableBindingOption Format,
+	const FTransform TargetDeformationSkeletonOffset, 
+	const FName TargetBone,
+	const float SimulationBlendWeight,
+	TArray<bool>* OutInfluence) const
+{
+	if (!GetRestCollection())
+	{
+		UE_LOG(LogDeformableTetrahedralComponentInternal, Warning,
+			TEXT("'%s' - GetSkeletalMeshEmbeddedPositionsInternal - RestCollection is not set."),
+			*GetName());
+		return TArray<FVector>();
+	}
+	if(!GetRestCollection()->TargetDeformationSkeleton)
+	{
+		UE_LOG(LogDeformableTetrahedralComponentInternal, Warning,
+			TEXT("'%s' - GetSkeletalMeshEmbeddedPositionsInternal - TargetDeformationSkeleton is not set on the flesh asset."),
+			*GetName());
+		return TArray<FVector>();
+	}
+
+	TArray<FVector> EmbeddedPosComp;
+
+	// Get sample points in the skel mesh's component space.  This code assumes that the
+	// skeletal mesh and the flesh asset are aligned in their respective local spaces.  If
+	// they're not aligned, then the TargetDeformationSkeletonOffset should be provided to
+	// put the skeletal mesh in the same place as the flesh mesh.
+	TArray<FVector> TransformPositions;
+	if (Format == ChaosDeformableBindingOption::WorldDelta || 
+		Format == ChaosDeformableBindingOption::ComponentDelta) // BoneDelta handled below
+	{
+		TArray<FTransform> ComponentPose;
+		Dataflow::Animation::GlobalTransforms(GetRestCollection()->TargetDeformationSkeleton->GetRefSkeleton(), ComponentPose);
+		TransformPositions.SetNumUninitialized(ComponentPose.Num());
+		if (TargetDeformationSkeletonOffset.Equals(FTransform::Identity))
+		{
+			for (int32 i = 0; i < ComponentPose.Num(); i++)
+			{
+				TransformPositions[i] = ComponentPose[i].GetTranslation();
+			}
+		}
+		else
+		{
+			for (int32 i = 0; i < ComponentPose.Num(); i++)
+			{
+				TransformPositions[i] =
+					TargetDeformationSkeletonOffset.TransformPosition(
+						ComponentPose[i].GetTranslation());
+			}
+		}
+	}
+	else
+	{
+		// If we aren't computing deltas, then we don't need the bone positions.  We
+		// only need to size the array to however many bones we have.
+		TransformPositions.SetNum(GetRestCollection()->TargetDeformationSkeleton->GetRefSkeleton().GetNum());
+	}
+
+
+
+	// World space
+	if (Format == ChaosDeformableBindingOption::WorldPos)
+	{
+		// Calculate their current embedded positions
+		EmbeddedPosComp =
+			GetEmbeddedPositionsInternal(
+			TransformPositions,
+			FName(GetRestCollection()->TargetDeformationSkeleton->GetName()), // for identifying the binding group
+			SimulationBlendWeight,
+			OutInfluence);
+
+		// Put component space points into world space
+		const FTransform ComponentXf = GetComponentTransform();
+		for (int32 i = 0; i < EmbeddedPosComp.Num(); i++)
+		{
+			EmbeddedPosComp[i] = ComponentXf.TransformPosition(EmbeddedPosComp[i]);
+		}
+	}
+	else if (Format == ChaosDeformableBindingOption::WorldDelta)
+	{
+		// Calculate their current embedded positions
+		EmbeddedPosComp =
+			GetEmbeddedPositionsInternal(
+			TransformPositions,
+			FName(GetRestCollection()->TargetDeformationSkeleton->GetName()), // for identifying the binding group
+			1.f,
+			OutInfluence);
+
+		const FTransform ComponentXf = GetComponentTransform();
+		for (int32 i = 0; i < EmbeddedPosComp.Num(); i++)
+		{
+			EmbeddedPosComp[i] = EmbeddedPosComp[i] - TransformPositions[i];
+			EmbeddedPosComp[i] = ComponentXf.TransformVector(EmbeddedPosComp[i]);
+		}
+	}
+
+	// Component space
+	else if (Format == ChaosDeformableBindingOption::ComponentPos)
+	{
+		EmbeddedPosComp =
+			GetEmbeddedPositionsInternal(
+			TransformPositions,
+			FName(GetRestCollection()->TargetDeformationSkeleton->GetName()), // for identifying the binding group
+			1.f,
+			OutInfluence);
+	}
+	else if (Format == ChaosDeformableBindingOption::ComponentDelta)
+	{
+		// Calculate their current embedded positions
+		EmbeddedPosComp =
+			GetEmbeddedPositionsInternal(
+			TransformPositions,
+			FName(GetRestCollection()->TargetDeformationSkeleton->GetName()), // for identifying the binding group
+			1.f,
+			OutInfluence);
+
+		for (int32 i = 0; i < EmbeddedPosComp.Num(); i++)
+		{
+			EmbeddedPosComp[i] = (EmbeddedPosComp[i] - TransformPositions[i]) * SimulationBlendWeight;
+		}
+	}
+
+	// Bone space
+	else
+	{
+		EmbeddedPosComp =
+			GetEmbeddedPositionsInternal(
+			TransformPositions,
+			FName(GetRestCollection()->TargetDeformationSkeleton->GetName()), // for identifying the binding group
+			SimulationBlendWeight,
+			OutInfluence);
+
+
+		// Find the component that owns TargetDeformationSkeleton, so we can pull the 
+		// current animated bone transforms out of it.
+		//
+		// It's possible that TargetDeformationSkeleton is owned by a component on another 
+		// actor, in which case, we can't (easily) find it. If that becomes a desired use
+		// case, then we'll need to have a handle to that component, not just the asset.
+		TArray<USkeletalMeshComponent*> SkeletalMeshComponents;
+		GetOwner()->GetComponents<USkeletalMeshComponent>(SkeletalMeshComponents);
+		USkeletalMeshComponent* TargetDeformationSkeletalMeshComponent = nullptr;
+		for (int32 i = 0; i < SkeletalMeshComponents.Num(); i++)
+		{
+			if (SkeletalMeshComponents[i]->GetSkeletalMeshAsset() == GetRestCollection()->TargetDeformationSkeleton)
+			{
+				TargetDeformationSkeletalMeshComponent = SkeletalMeshComponents[i];
+				break;
+			}
+		}
+
+		if (TargetDeformationSkeletalMeshComponent)
+		{
+			// Get the current animated bone transforms and then their positions.
+			const TArray<FTransform>& TargetDeformationComponentTransforms = 
+				TargetDeformationSkeletalMeshComponent->GetComponentSpaceTransforms();
+			TArray<FVector> AnimTransformPositions;
+			AnimTransformPositions.SetNumUninitialized(TargetDeformationComponentTransforms.Num());
+
+			// Get their positions, and apply the offset if necessary.
+			if (!TargetDeformationSkeletonOffset.Equals(FTransform::Identity) && 
+				Format == ChaosDeformableBindingOption::BoneDelta)
+			{
+				for (int32 i = 0; i < TargetDeformationComponentTransforms.Num(); i++)
+				{
+					AnimTransformPositions[i] =
+						TargetDeformationSkeletonOffset.TransformPosition(
+							TargetDeformationComponentTransforms[i].GetTranslation());
+				}
+			}
+			else
+			{
+				for (int32 i = 0; i < TargetDeformationComponentTransforms.Num(); i++)
+				{
+					AnimTransformPositions[i] = TargetDeformationComponentTransforms[i].GetTranslation();
+				}
+			}
+
+			// Find the transform index of 'TargetBone'
+			FTransform BoneToComponentXf = FTransform::Identity;
+			const int32 BoneIndex = GetRestCollection()->TargetDeformationSkeleton->GetRefSkeleton().FindBoneIndex(TargetBone);
+			if (TargetDeformationComponentTransforms.IsValidIndex(BoneIndex))
+			{
+				BoneToComponentXf = FTransform(TargetDeformationComponentTransforms[BoneIndex].ToMatrixWithScale().Inverse());
+			}
+			else
+			{
+				UE_LOG(LogDeformableTetrahedralComponentInternal, Warning, 
+					TEXT("'%s' - Failed to find a valid bone index (got %d) for bone name '%s' in TargetDeformationSkeleton '%s' "
+					"corresponding to SkeletalMeshComponent '%s', which has %d bones."), 
+					*GetName(),
+					BoneIndex, *TargetBone.ToString(), *GetRestCollection()->TargetDeformationSkeleton.GetName(),
+					*TargetDeformationSkeletalMeshComponent->GetName(), TargetDeformationComponentTransforms.Num());
+				return TArray<FVector>();
+			}
+
+			// Compute the return values
+			if (Format == ChaosDeformableBindingOption::BonePos)
+			{
+				if (!BoneToComponentXf.Equals(FTransform::Identity))
+				{
+					for (int32 i = 0; i < EmbeddedPosComp.Num(); i++)
+					{
+						EmbeddedPosComp[i] = BoneToComponentXf.TransformPosition(EmbeddedPosComp[i]);
+					}
+				}
+			}
+			else if (Format == ChaosDeformableBindingOption::BoneDelta)
+			{
+				if (!BoneToComponentXf.Equals(FTransform::Identity))
+				{
+					for (int32 i = 0; i < EmbeddedPosComp.Num(); i++)
+					{
+						EmbeddedPosComp[i] = EmbeddedPosComp[i] - AnimTransformPositions[i];
+						EmbeddedPosComp[i] = BoneToComponentXf.TransformVector(EmbeddedPosComp[i]);
+					}
+				}
+				else
+				{
+					for (int32 i = 0; i < EmbeddedPosComp.Num(); i++)
+					{
+						EmbeddedPosComp[i] = EmbeddedPosComp[i] - AnimTransformPositions[i];
+					}
+				}
+			}
+		}
+		else
+		{
+			UE_LOG(LogDeformableTetrahedralComponentInternal, Warning, 
+				TEXT("'%s' - Failed to find SkeletalMeshComponent for TargetDeformationSkeleton '%s'."),
+				*GetName(),
+				*GetRestCollection()->TargetDeformationSkeleton->GetName());
+			return TArray<FVector>();
+		}
+	}
+	return EmbeddedPosComp;
+}
+
+TArray<FVector> UDeformableTetrahedralComponent::GetEmbeddedPositionsInternal(
+	const TArray<FVector>& InPositions, 
+	const FName SkeletalMeshName,
+	const float SimulationBlendWeight,
+	TArray<bool>* OutInfluence) const
+{
+	auto UEVert3d = [](FVector3f V) { return FVector3d(V.X, V.Y, V.Z); };
+	auto UEVert4d = [](FVector4f V) { return FVector4d(V.X, V.Y, V.Z, V.W); };
+
+	TArray<FVector> OutPositions = InPositions;
+
+	if (const UFleshAsset* RestAsset = GetRestCollection())
+	{
+		if (const FFleshCollection* Rest = RestAsset->GetCollection())
+		{
+			GeometryCollection::Facades::FTetrahedralSkeletalBindings TetBindings(*Rest);
+
+			const TManagedArray<int32>* TetrahedronStart = Rest->FindAttribute<int32>(FTetrahedralCollection::TetrahedronStartAttribute, FGeometryCollection::GeometryGroup);
+			const TManagedArray<FVector3f>* RestVerts = Rest->FindAttributeTyped<FVector3f>("Vertex", FGeometryCollection::VerticesGroup);
+			const TManagedArray<FVector3f>* Verts = GetDynamicCollection() ? GetDynamicCollection()->FindPositions() : RestVerts;
+
+			if (ensure(Verts != nullptr) && TetrahedronStart)
+			{
+				auto CalculateBindings = [this, TetBindings, TetrahedronStart, SkeletalMeshName, OutInfluence](const TManagedArray<FVector3f>* Verts, TArray<FVector>& OutPositions)
+				{
+					if (OutInfluence != nullptr)
+						OutInfluence->Init(false, OutPositions.Num());
+					for (int32 TetMeshIdx = 0; TetMeshIdx < TetrahedronStart->Num(); TetMeshIdx++)
+					{
+						FString MeshBindingsName =
+							GeometryCollection::Facades::FTetrahedralSkeletalBindings::GenerateMeshGroupName(TetMeshIdx, SkeletalMeshName);
+						if (!TetBindings.CalculateBindings(MeshBindingsName, Verts->GetConstArray(), OutPositions, OutInfluence))
+						{
+							UE_LOG(LogDeformableTetrahedralComponentInternal, Warning,
+								   TEXT("'%s' - GetEmbeddedPositionsInternal - Failed to find mesh bindings for skeletal mesh '%s'"),
+								   *GetName(),
+								   *SkeletalMeshName.ToString());
+						}
+					}
+				};
+				CalculateBindings(Verts, OutPositions);
+
+				// blend between the aniamtion position and the 
+				// simulated position. 
+				if (!FMath::IsNearlyEqual(SimulationBlendWeight, 1.0) && GetDynamicCollection())
+				{
+					TArray<FVector> RestPositions = InPositions;
+					CalculateBindings(RestVerts, RestPositions);
+
+					float ClampedWeight = FMath::Clamp(SimulationBlendWeight, 0.f, 1.f);
+					for (int i = 0; i < OutPositions.Num(); i++)
+					{
+						FVector V = RestPositions[i] - OutPositions[i];
+						OutPositions[i] += V * (1.0-ClampedWeight);
+					}
+				}
+			}
+		}
+	}
+	return OutPositions;
+}
+
+TArray<FVector> UDeformableTetrahedralComponent::GetSkeletalMeshBindingPositionsInternal(const USkeletalMesh* InSkeletalMesh, TArray<bool>* OutInfluence) const
+{
+	auto UEVert3d = [](FVector3f V) { return FVector3d(V.X, V.Y, V.Z); };
+	auto UEVert4d = [](FVector4f V) { return FVector4d(V.X, V.Y, V.Z, V.W); };
+
+	TArray<FVector> TransformPositions;
+	if (InSkeletalMesh)
+	{
+		FName SkeletalMeshName(InSkeletalMesh->GetName());
+		if (const UFleshAsset* RestAsset = GetRestCollection())
+		{
+			if (const FFleshCollection* Rest = RestAsset->GetCollection())
+			{
+				GeometryCollection::Facades::FTetrahedralSkeletalBindings TetBindings(*Rest);
+
+				const TManagedArray<int32>* TetrahedronStart = Rest->FindAttribute<int32>(FTetrahedralCollection::TetrahedronStartAttribute, FGeometryCollection::GeometryGroup);
+				const TManagedArray<FVector3f>* Verts = GetDynamicCollection() ? GetDynamicCollection()->FindPositions() : Rest->FindAttributeTyped<FVector3f>("Vertex", FGeometryCollection::VerticesGroup);
+
+				if (ensure(Verts != nullptr) && TetrahedronStart)
+				{
+					// Component relative transforms, not world.
+					TArray<FTransform> ComponentPose;
+					Dataflow::Animation::GlobalTransforms(InSkeletalMesh->GetRefSkeleton(), ComponentPose);
+
+					TransformPositions.SetNumUninitialized(ComponentPose.Num());
+					for (int32 i = 0; i < ComponentPose.Num(); i++)
+					{
+						TransformPositions[i] = ComponentPose[i].GetTranslation();
+					}
+
+					if (OutInfluence != nullptr) OutInfluence->Init(false, TransformPositions.Num());
+					for (int32 TetMeshIdx = 0; TetMeshIdx < TetrahedronStart->Num(); TetMeshIdx++)
+					{
+						FString MeshBindingsName = GeometryCollection::Facades::FTetrahedralSkeletalBindings::GenerateMeshGroupName(TetMeshIdx, SkeletalMeshName);
+						TetBindings.CalculateBindings(MeshBindingsName, Verts->GetConstArray(), TransformPositions, OutInfluence);
+					}
+				}
+			}
+		}
+	}
+	return TransformPositions;
+}
+
+
diff --git a/Engine/Plugins/Experimental/ChaosFlesh/Source/ChaosFleshEngine/Private/ChaosFlesh/FleshActor.cpp b/Engine/Plugins/Experimental/ChaosFlesh/Source/ChaosFleshEngine/Private/ChaosFlesh/FleshActor.cpp
--- a/Engine/Plugins/Experimental/ChaosFlesh/Source/ChaosFleshEngine/Private/ChaosFlesh/FleshActor.cpp
+++ b/Engine/Plugins/Experimental/ChaosFlesh/Source/ChaosFleshEngine/Private/ChaosFlesh/FleshActor.cpp
@@ -58,9 +58,9 @@ void AFleshActor::PostEditChangeProperty(struct FPropertyChangedEvent& PropertyC
 				if (FleshComponent)
 				{
 					FleshComponent->PrimarySolverComponent = SolverComponent;
-					if (!SolverComponent->DeformableComponents.Contains(FleshComponent))
+					if (!SolverComponent->ConnectedObjects.DeformableComponents.Contains(FleshComponent))
 					{
-						SolverComponent->DeformableComponents.Add(TObjectPtr<UDeformablePhysicsComponent>(FleshComponent));
+						SolverComponent->ConnectedObjects.DeformableComponents.Add(TObjectPtr<UDeformablePhysicsComponent>(FleshComponent));
 					}
 				}
 			}
@@ -72,9 +72,9 @@ void AFleshActor::PostEditChangeProperty(struct FPropertyChangedEvent& PropertyC
 				if (FleshComponent)
 				{
 					FleshComponent->PrimarySolverComponent = nullptr;
-					if (SolverComponent->DeformableComponents.Contains(FleshComponent))
+					if (SolverComponent->ConnectedObjects.DeformableComponents.Contains(FleshComponent))
 					{
-						SolverComponent->DeformableComponents.Remove(FleshComponent);
+						SolverComponent->ConnectedObjects.DeformableComponents.Remove(FleshComponent);
 					}
 				}
 			}
diff --git a/Engine/Plugins/Experimental/ChaosFlesh/Source/ChaosFleshEngine/Private/ChaosFlesh/FleshComponent.cpp b/Engine/Plugins/Experimental/ChaosFlesh/Source/ChaosFleshEngine/Private/ChaosFlesh/FleshComponent.cpp
--- a/Engine/Plugins/Experimental/ChaosFlesh/Source/ChaosFleshEngine/Private/ChaosFlesh/FleshComponent.cpp
+++ b/Engine/Plugins/Experimental/ChaosFlesh/Source/ChaosFleshEngine/Private/ChaosFlesh/FleshComponent.cpp
@@ -2,995 +2,13 @@
 
 #include "ChaosFlesh/FleshComponent.h"
 
-#include "Animation/SkeletalMeshActor.h"
-#include "Animation/Skeleton.h"
-#include "ChaosStats.h"
-#include "Chaos/DebugDrawQueue.h"
-#include "Chaos/Tetrahedron.h"
-#include "ChaosFlesh/FleshCollection.h"
-#include "ChaosFlesh/ChaosDeformableSolverActor.h"
-#include "ChaosFlesh/ChaosDeformableSolverComponent.h"
-#include "ChaosFlesh/ChaosDeformableTypes.h"
-#include "ChaosFlesh/FleshDynamicAsset.h"
-#include "Components/SkeletalMeshComponent.h"
-#include "Dataflow/DataflowEngineUtil.h"
-#include "Engine/SkeletalMesh.h"
-#include "GeometryCollection/Facades/CollectionTransformSourceFacade.h"
-#include "GeometryCollection/Facades/CollectionTetrahedralSkeletalBindingsFacade.h"
-#include "GeometryCollection/TransformCollection.h"
-#include "ProceduralMeshComponent.h"
-#include "Rendering/SkeletalMeshRenderData.h"
-
 
 #include UE_INLINE_GENERATED_CPP_BY_NAME(FleshComponent)
 
-FChaosEngineDeformableCVarParams CVarParams;
-FAutoConsoleVariableRef CVarDeforambleDoDrawSimulationMesh(TEXT("p.Chaos.DebugDraw.Deformable.SimulationMesh"), CVarParams.bDoDrawSimulationMesh, TEXT("Debug draw the deformable simulation resutls on the game thread. [def: true]"));
-FAutoConsoleVariableRef CVarDeforambleDoDrawSkeletalMeshBindingPositions(TEXT("p.Chaos.DebugDraw.Deformable.SkeletalMeshBindingPositions"), CVarParams.bDoDrawSkeletalMeshBindingPositions, TEXT("Debug draw the deformable simulation's SkeletalMeshBindingPositions on the game thread. [def: false]"));
-FAutoConsoleVariableRef CVarDeforambleDoDrawSkeletalMeshBindingPositionsSimulationBlendWeight(TEXT("p.Chaos.DebugDraw.Deformable.SkeletalMeshBindingPositions.SimulationBlendWeight"), CVarParams.DrawSkeletalMeshBindingPositionsSimulationBlendWeight, TEXT("Set the simulation blend weight of the skeletal mesh debug draw.[def: 1.]"));
-
-
-#define PERF_SCOPE(X) SCOPE_CYCLE_COUNTER(X); TRACE_CPUPROFILER_EVENT_SCOPE(X);
-DECLARE_CYCLE_STAT(TEXT("Chaos.Deformable.UFleshComponent.TickComponent"), STAT_ChaosDeformable_UFleshComponent_TickComponent, STATGROUP_Chaos);
-DECLARE_CYCLE_STAT(TEXT("Chaos.Deformable.UFleshComponent.NewProxy"), STAT_ChaosDeformable_UFleshComponent_NewProxy, STATGROUP_Chaos);
-DECLARE_CYCLE_STAT(TEXT("Chaos.Deformable.UFleshComponent.NewDeformableData"), STAT_ChaosDeformable_UFleshComponent_NewDeformableData, STATGROUP_Chaos);
-DECLARE_CYCLE_STAT(TEXT("Chaos.Deformable.UFleshComponent.UpdateFromSimualtion"), STAT_ChaosDeformable_UFleshComponent_UpdateFromSimualtion, STATGROUP_Chaos);
-DECLARE_CYCLE_STAT(TEXT("Chaos.Deformable.UFleshComponent.RenderProceduralMesh"), STAT_ChaosDeformable_UFleshComponent_RenderProceduralMesh, STATGROUP_Chaos);
-
 DEFINE_LOG_CATEGORY_STATIC(LogFleshComponentInternal, Log, All);
 
 UFleshComponent::UFleshComponent(const FObjectInitializer& ObjectInitializer)
 	: Super(ObjectInitializer)
 {
-	Mesh = ObjectInitializer.CreateDefaultSubobject<UProceduralMeshComponent>(this, TEXT("Flesh Visualization Component"));
-	PrimaryComponentTick.TickGroup = TG_LastDemotable;
-	PrimaryComponentTick.bCanEverTick = CVarParams.bDoDrawSimulationMesh;
-	bTickInEditor = CVarParams.bDoDrawSimulationMesh;
-
-	DynamicCollection = ObjectInitializer.CreateDefaultSubobject<UFleshDynamicAsset>(this, TEXT("Flesh Dynamic Asset"));
-	SimulationCollection = ObjectInitializer.CreateDefaultSubobject<USimulationAsset>(this, TEXT("Flesh Simulation Asset"));
-}
-
-UFleshComponent::~UFleshComponent()
-{
-	if (RenderMesh) delete RenderMesh;
-}
-
-void UFleshComponent::Invalidate()
-{
-	bBoundsNeedsUpdate = true;
-}
-
-
-void UFleshComponent::OnRegister()
-{
-	if (bBoundsNeedsUpdate)
-	{
-		UpdateLocalBounds();
-	}
-	Super::OnRegister();
-}
-
-void UFleshComponent::EndPlay(const EEndPlayReason::Type ReasonEnd)
-{
-	if (GetDynamicCollection())
-	{
-		GetDynamicCollection()->Reset();
-	}
-	if (GetSimulationCollection())
-	{
-		GetSimulationCollection()->Reset();
-	}
-}
-
-void UFleshComponent::SetRestCollection(const UFleshAsset* InRestCollection)
-{
-	RestCollection = InRestCollection;
-	Invalidate();
-	UpdateLocalBounds();
-	ResetProceduralMesh();
-}
-
-
-UDeformablePhysicsComponent::FThreadingProxy* UFleshComponent::NewProxy()
-{
-	PERF_SCOPE(STAT_ChaosDeformable_UFleshComponent_NewProxy);
-
-	UpdateSimSpaceTransformIndex();
-	if (const UFleshAsset* RestAsset = GetRestCollection())
-	{
-		if (const FFleshCollection* Rest = RestAsset->GetCollection())
-		{
-			if (Rest->NumElements(FGeometryCollection::VerticesGroup))
-			{
-				if (!GetDynamicCollection())
-				{
-					DynamicCollection = NewObject<UFleshDynamicAsset>(this, TEXT("Flesh Dynamic Asset"));
-				}
-				if (!GetSimulationCollection())
-				{
-					SimulationCollection = NewObject<USimulationAsset>(this, TEXT("Flesh Simulation Asset"));
-				}
-
-				GetDynamicCollection()->Reset(Rest);
-				if (const FManagedArrayCollection* Dynamic = GetDynamicCollection()->GetCollection())
-				{
-					// Mesh points are in component space, such that the exterior hull aligns with the
-					// surface of the skeletal mesh, which is subject to the transform hierarchy.
-					const FTransform& ComponentToWorldXf = GetComponentTransform();
-					const FTransform ComponentToSimXf = GetSimSpaceRestTransform();
-					return new FFleshThreadingProxy(
-						this,
-						ComponentToWorldXf,
-						ComponentToSimXf,
-						SimSpace,
-						*Rest,
-						*Dynamic);
-				}
-			}
-		}
-	}
-	return nullptr;
-}
-
-UDeformablePhysicsComponent::FDataMapValue UFleshComponent::NewDeformableData()
-{
-	PERF_SCOPE(STAT_ChaosDeformable_UFleshComponent_NewDeformableData);
-
-	using namespace GeometryCollection::Facades;
-	if (GetOwner())
-	{
-		if (const UFleshAsset* FleshAsset = GetRestCollection())
-		{
-			if (const FFleshCollection* Rest = FleshAsset->GetCollection())
-			{
-				FTransformSource TransformSource(*Rest);
-				if (TransformSource.IsValid())
-				{
-					TArray<USkeletalMeshComponent*> SkeletalMeshComponents;
-					GetOwner()->GetComponents<USkeletalMeshComponent>(SkeletalMeshComponents);
-
-					if (const TManagedArray<FTransform>* RestTransforms = Rest->FindAttribute<FTransform>(FTransformCollection::TransformAttribute, FTransformCollection::TransformGroup))
-					{
-						TArray<FTransform> AnimationTransforms = RestTransforms->GetConstArray();
-						TArray<FTransform> ComponentPose = RestTransforms->GetConstArray();
-
-						// Extract animated transforms from all skeletal meshes.
-						for (const USkeletalMeshComponent* SkeletalMeshComponent : SkeletalMeshComponents)
-						{
-							if (const USkeletalMesh* SkeletalMesh = SkeletalMeshComponent->GetSkeletalMeshAsset())
-							{
-								if (const USkeleton* Skeleton = SkeletalMesh->GetSkeleton())
-								{
-									TSet<int32> Roots = TransformSource.GetTransformSource(Skeleton->GetName(), Skeleton->GetGuid().ToString());
-									if (!Roots.IsEmpty() && ensureMsgf(Roots.Num() == 1, TEXT("Error: Only supports a single root per skeleton.(%s)"), *Skeleton->GetName()))
-									{
-										TArray<FTransform> ComponentLocalPose;
-										Dataflow::Animation::GlobalTransforms(SkeletalMesh->GetRefSkeleton(), ComponentLocalPose);
-
-										const TArray<FTransform>& ComponentTransforms = SkeletalMeshComponent->GetComponentSpaceTransforms();
-										if (ComponentLocalPose.Num() == ComponentTransforms.Num())
-										{
-											for (int Adx = Roots.Array()[0], Cdx = 0; Adx <= AnimationTransforms.Num() && Cdx < ComponentTransforms.Num(); Adx++, Cdx++)
-											{
-												// @todo(flesh) : Can we just use one array?
-												AnimationTransforms[Adx] = ComponentTransforms[Cdx];
-												ComponentPose[Adx] = ComponentLocalPose[Cdx]; 
-
-												if (SimSpaceTransformGlobalIndex == INDEX_NONE &&
-													SimSpaceTransformIndex == Cdx &&
-													SkeletalMesh == SimSpaceSkeletalMesh)
-												{
-													SimSpaceTransformGlobalIndex = Adx;
-												}
-											}
-
-
-										}
-									}
-								}
-							}
-						}
-
-						FTransform BoneSpaceXf;
-						if (AnimationTransforms.IsValidIndex(SimSpaceTransformGlobalIndex))
-						{
-							BoneSpaceXf = AnimationTransforms[SimSpaceTransformGlobalIndex];
-						}
-						else
-						{
-							BoneSpaceXf = FTransform::Identity;
-						}
-
-						return FDataMapValue(
-							new Chaos::Softs::FFleshThreadingProxy::FFleshInputBuffer(
-								*GetSimulationCollection()->GetCollection(),
-								this->GetComponentTransform(),
-								BoneSpaceXf, 
-								SimSpaceTransformGlobalIndex,
-								AnimationTransforms, 
-								ComponentPose, 
-								bTempEnableGravity, 
-								StiffnessMultiplier, 
-								DampingMultiplier, 
-								MassMultiplier, 
-								IncompressibilityMultiplier,
-								InflationMultiplier,
-								this));
-					}
-				}
-			}
-		}
-	}
-	return FDataMapValue(
-		new Chaos::Softs::FFleshThreadingProxy::FFleshInputBuffer(
-			*GetSimulationCollection()->GetCollection(),
-			this->GetComponentTransform(),
-			GetSimSpaceRestTransform(),
-			SimSpaceTransformGlobalIndex,
-			bTempEnableGravity,
-			StiffnessMultiplier, 
-			DampingMultiplier, 
-			MassMultiplier, 
-			IncompressibilityMultiplier,
-			InflationMultiplier,
-			this));
-}
-
-TArray<FString> UFleshComponent::GetSimSpaceBoneNameOptions() const
-{
-	TArray<FString> Names;
-	if (RestCollection)
-	{
-		if (RestCollection->SkeletalMesh)
-		{
-			const FReferenceSkeleton& RefSkeleton = RestCollection->SkeletalMesh->GetSkeleton()->GetReferenceSkeleton();
-			Names.SetNum(RefSkeleton.GetNum());
-			for (int32 i = 0; i < RefSkeleton.GetNum(); i++)
-			{
-				Names[i] = RefSkeleton.GetBoneName(i).ToString();
-			}
-		}
-	}
-	return Names;
-}
-
-bool UFleshComponent::UpdateSimSpaceTransformIndex()
-{
-	SimSpaceTransformIndex = INDEX_NONE;
-	SimSpaceSkeletalMesh = nullptr;
-
-	if (SimSpace != ChaosDeformableSimSpace::Bone)
-	{
-		return false;
-	}
-
-	if (RestCollection && RestCollection->SkeletalMesh)
-	{
-		const FReferenceSkeleton& RefSkeleton = RestCollection->SkeletalMesh->GetSkeleton()->GetReferenceSkeleton();
-		for (int32 i = 0; i < RefSkeleton.GetNum(); i++)
-		{
-			if (RefSkeleton.GetBoneName(i).ToString() == SimSpaceBoneName.ToString())
-			{
-				SimSpaceSkeletalMesh = RestCollection->SkeletalMesh;
-				SimSpaceTransformIndex = i;
-				return true;
-			}
-		}
-	}
-	return false;
-}
-
-FTransform UFleshComponent::GetSimSpaceRestTransform() const
-{
-	if (SimSpaceSkeletalMesh == nullptr)
-	{
-		return FTransform::Identity;
-	}
-	
-	TArray<FTransform> ComponentTransforms;
-	ComponentTransforms.SetNum(SimSpaceSkeletalMesh->GetRefSkeleton().GetNum());
-
-	SimSpaceSkeletalMesh->FillComponentSpaceTransforms(
-		SimSpaceSkeletalMesh->GetRefSkeleton().GetRefBonePose(),
-		SimSpaceSkeletalMesh->GetResourceForRendering()->LODRenderData[0].RequiredBones,
-		ComponentTransforms);
-
-	if (!ComponentTransforms.IsValidIndex(SimSpaceTransformIndex))
-	{
-		return FTransform::Identity;
-	}
-	const FTransform& ComponentToBone = ComponentTransforms[SimSpaceTransformIndex];
-	return ComponentToBone;
-}
-
-void UFleshComponent::UpdateFromSimualtion(const FDataMapValue* SimualtionBuffer)
-{
-	PERF_SCOPE(STAT_ChaosDeformable_UFleshComponent_UpdateFromSimualtion);
-
-	if (const FFleshThreadingProxy::FFleshOutputBuffer* FleshBuffer = (*SimualtionBuffer)->As<FFleshThreadingProxy::FFleshOutputBuffer>())
-	{
-		if (GetDynamicCollection())
-		{
-			// @todo(flesh) : reduce conversions
-			auto UEVertd = [](FVector3f V) { return FVector3d(V.X, V.Y, V.Z); };
-			auto UEVertf = [](FVector3d V) { return FVector3f(V.X, V.Y, V.Z); };
-
-			TManagedArray<FVector3f>& DynamicVertex = GetDynamicCollection()->GetPositions();
-			const TManagedArray<FVector3f>& SimulationVertex = FleshBuffer->Dynamic.GetAttribute<FVector3f>("Vertex", FGeometryCollection::VerticesGroup);
-
-			// Simulator produces results in component space.
-			for (int i = DynamicVertex.Num() - 1; i >= 0; i--)
-			{
-				DynamicVertex[i] = SimulationVertex[i];
-			}
-			
-			//p.Chaos.DebugDraw.Enabled 1
-			//p.Chaos.DebugDraw.Deformable.SkeletalMeshBindingPositions 1
-			if (CVarParams.bDoDrawSkeletalMeshBindingPositions)
-			{
-				DebugDrawSkeletalMeshBindingPositions();
-			}
-		}
-	}
-}
-
-void UFleshComponent::UpdateLocalBounds()
-{
-	if (bBoundsNeedsUpdate && RestCollection)
-	{
-		{
-			FFleshAssetEdit EditObject = RestCollection->EditCollection();
-			if (FFleshCollection* Collection = EditObject.GetFleshCollection())
-			{
-				Collection->UpdateBoundingBox();
-			}
-		}
-		BoundingBox = RestCollection->GetCollection()->GetBoundingBox();
-		bBoundsNeedsUpdate = false;
-	}
-
-}
-
-FBoxSphereBounds UFleshComponent::CalcBounds(const FTransform& LocalToWorldIn) const
-{
-	return BoundingBox.TransformBy(GetComponentTransform()); // todo(chaos:flesh) use LocalToWorldIn
-}
-
-void UFleshComponent::ResetDynamicCollection()
-{
-	if (const UFleshAsset* RestAsset = GetRestCollection())
-	{
-		if (!GetDynamicCollection())
-		{
-			DynamicCollection = NewObject<UFleshDynamicAsset>(this, TEXT("Flesh Dynamic Asset"));
-		}
-
-		if (!GetDynamicCollection()->GetCollection() || 
-			!GetDynamicCollection()->GetCollection()->NumElements(FGeometryCollection::VerticesGroup))
-		{
-			GetDynamicCollection()->Reset(RestAsset->GetCollection());
-			ResetProceduralMesh();
-		}
-		else
-		{
-			GetDynamicCollection()->ResetAttributesFrom(RestAsset->GetCollection());
-		}
-	}
-}
-
-//
-// Rendering Support
-//
-
-void UFleshComponent::BeginPlay()
-{
-	Super::BeginPlay();
-
-	if (ASkeletalMeshActor* SkeletalMeshActor = Cast<ASkeletalMeshActor>(GetOwner()))
-	{
-		if (USkeletalMeshComponent* SkeletalMeshComponent = SkeletalMeshActor->GetSkeletalMeshComponent())
-		{
-			PrimaryComponentTick.AddPrerequisite(SkeletalMeshComponent, SkeletalMeshComponent->PrimaryComponentTick);
-		}
-	}
-	if (PrimarySolverComponent)
-	{
-		PrimaryComponentTick.AddPrerequisite(PrimarySolverComponent, PrimarySolverComponent->PrimaryComponentTick);
-	}
-}
-
-void UFleshComponent::TickComponent(float DeltaTime, enum ELevelTick TickType, FActorComponentTickFunction* ThisTickFunction)
-{
-	PERF_SCOPE(STAT_ChaosDeformable_UFleshComponent_TickComponent);
-	Super::TickComponent(DeltaTime, TickType, ThisTickFunction);
-
-	if (IsVisible())
-	{
-		RenderProceduralMesh();
-	}
-	else
-	{
-		ResetProceduralMesh();
-	}
-}
-
-void UFleshComponent::ResetProceduralMesh()
-{
-	if (Mesh && RenderMesh)
-	{
-		Mesh->ClearAllMeshSections();
-	}
-	if (RenderMesh)
-	{
-		delete RenderMesh;
-		RenderMesh = nullptr;
-	}
-}
-
-void UFleshComponent::RenderProceduralMesh()
-{
-	bool bCanRender = false;
-	if (const UFleshAsset* FleshAsset = GetRestCollection())
-	{
-		if (IsVisible())
-		{
-#if WITH_EDITORONLY_DATA
-			if (FleshAsset->bRenderInEditor)
-#endif
-			{
-				if (Mesh && CVarParams.bDoDrawSimulationMesh)
-				{
-					PERF_SCOPE(STAT_ChaosDeformable_UFleshComponent_RenderProceduralMesh);
-
-					if (const FFleshCollection* Flesh = FleshAsset->GetCollection())
-					{
-						if (HideTetrahedra.Num() != NumSkippedTets)
-						{
-							const TManagedArray<FIntVector4>* Tets =
-								Flesh->FindAttribute<FIntVector4>(
-									FTetrahedralCollection::TetrahedronAttribute, FTetrahedralCollection::TetrahedralGroup);
-
-							FaceVerticesToSkip.Reset();
-							for (TSet<int32>::TConstIterator It = HideTetrahedra.CreateConstIterator(); It; ++It)
-							{
-								const int32 Idx = *It;
-								if (Idx >= 0 && Idx < Tets->Num())
-								{
-									const FIntVector4& Tet = (*Tets)[Idx];
-									for (int32 i = 0; i < 4; i++)
-									{
-										FaceVerticesToSkip.Add(Tet[i]);
-									}
-								}
-							}
-							NumSkippedTets = HideTetrahedra.Num();
-							NumSkippedFaces = 0;
-						}
-
-						int32 NumVertices = Flesh->NumElements(FGeometryCollection::VerticesGroup);
-						int32 NumFaces = Flesh->NumElements(FGeometryCollection::FacesGroup);
-						if (NumFaces && NumVertices)
-						{
-							if (RenderMesh && RenderMesh->Vertices.Num() != (NumFaces - NumSkippedFaces) * 3)
-							{
-								ResetProceduralMesh();
-							}
-
-							// Find vertices array.  Use dynamic collection if available and size matches.
-							const TManagedArray<FVector3f>* RenderVertex = &Flesh->Vertex;
-							if (GetDynamicCollection())
-							{
-								const TManagedArray<FVector3f>& DynamicVertex = GetDynamicCollection()->GetPositions();
-								if (DynamicVertex.Num() == RenderVertex->Num())
-								{
-									RenderVertex = &DynamicVertex;
-								}
-							}
-
-							const TManagedArray<bool>* FacesVisible =
-								Flesh->FindAttribute<bool>(
-									"Visible", FGeometryCollection::FacesGroup);
-
-							if (!RenderMesh)
-							{
-								RenderMesh = new FFleshRenderMesh;
-								NumSkippedFaces = 0;
-
-								for (int i = 0, j = 0; i < NumFaces; ++i)
-								{
-									if (((FacesVisible && (*FacesVisible)[i] == false)) ||
-										(FaceVerticesToSkip.Contains(Flesh->Indices[i][0]) &&
-										 FaceVerticesToSkip.Contains(Flesh->Indices[i][1]) &&
-										 FaceVerticesToSkip.Contains(Flesh->Indices[i][2])))
-									{
-										NumSkippedFaces++;
-										continue;
-									}
-
-									const auto& P1 = (*RenderVertex)[Flesh->Indices[i][0]];
-									const auto& P2 = (*RenderVertex)[Flesh->Indices[i][1]];
-									const auto& P3 = (*RenderVertex)[Flesh->Indices[i][2]];
-
-									RenderMesh->Vertices.Add(FVector(P1));
-									RenderMesh->Vertices.Add(FVector(P2));
-									RenderMesh->Vertices.Add(FVector(P3));
-
-									RenderMesh->Colors.Add(FLinearColor::White);
-									RenderMesh->Colors.Add(FLinearColor::White);
-									RenderMesh->Colors.Add(FLinearColor::White);
-
-									RenderMesh->UVs.Add(FVector2D(0, 0));
-									RenderMesh->UVs.Add(FVector2D(0, 0));
-									RenderMesh->UVs.Add(FVector2D(0, 0));
-
-									RenderMesh->Triangles.Add(3 * j);
-									RenderMesh->Triangles.Add(3 * j + 1);
-									RenderMesh->Triangles.Add(3 * j + 2);
-
-									auto Normal = -Chaos::FVec3::CrossProduct(P3 - P1, P2 - P1);
-									RenderMesh->Normals.Add(Normal);
-									RenderMesh->Normals.Add(Normal);
-									RenderMesh->Normals.Add(Normal);
-
-									auto Tangent = (P2 - P1).GetSafeNormal();
-									RenderMesh->Tangents.Add(FProcMeshTangent(Tangent[0], Tangent[1], Tangent[2]));
-									Tangent = (P3 - P2).GetSafeNormal();
-									RenderMesh->Tangents.Add(FProcMeshTangent(Tangent[0], Tangent[1], Tangent[2]));
-									Tangent = (P1 - P3).GetSafeNormal();
-									RenderMesh->Tangents.Add(FProcMeshTangent(Tangent[0], Tangent[1], Tangent[2]));
-								
-									j++;
-								}
-
-								Mesh->SetRelativeTransform(GetComponentTransform());
-								Mesh->CreateMeshSection_LinearColor(0, RenderMesh->Vertices, RenderMesh->Triangles, RenderMesh->Normals, RenderMesh->UVs, RenderMesh->Colors, RenderMesh->Tangents, false);
-							}
-							else
-							{
-								auto InRange = [](int32 Size, int32 Val) { return 0 <= Val && Val < Size; };
-
-								// Display only
-								for (int i = 0, j = 0; i < NumFaces; ++i)
-								{
-									if (((FacesVisible && (*FacesVisible)[i] == false)) ||
-										(FaceVerticesToSkip.Contains(Flesh->Indices[i][0]) &&
-										 FaceVerticesToSkip.Contains(Flesh->Indices[i][1]) &&
-										 FaceVerticesToSkip.Contains(Flesh->Indices[i][2])))
-									{
-										continue;
-									}
-
-									const auto& P1 = (*RenderVertex)[Flesh->Indices[i][0]];
-									const auto& P2 = (*RenderVertex)[Flesh->Indices[i][1]];
-									const auto& P3 = (*RenderVertex)[Flesh->Indices[i][2]];
-
-									RenderMesh->Vertices[3 * j] = FVector(P1);
-									RenderMesh->Vertices[3 * j + 1] = FVector(P2);
-									RenderMesh->Vertices[3 * j + 2] = FVector(P3);
-
-									auto Normal = Chaos::FVec3::CrossProduct(P3 - P1, P2 - P1);
-									RenderMesh->Normals[3 * j] = Normal;
-									RenderMesh->Normals[3 * j + 1] = Normal;
-									RenderMesh->Normals[3 * j + 2] = Normal;
-
-									auto Tangent = (P2 - P1).GetSafeNormal();
-									RenderMesh->Tangents[3 * j] = FProcMeshTangent(Tangent[0], Tangent[1], Tangent[2]);
-									Tangent = (P3 - P2).GetSafeNormal();
-									RenderMesh->Tangents[3 * j + 1] = FProcMeshTangent(Tangent[0], Tangent[1], Tangent[2]);
-									Tangent = (P1 - P3).GetSafeNormal();
-									RenderMesh->Tangents[3 * j + 2] = FProcMeshTangent(Tangent[0], Tangent[1], Tangent[2]);
-
-									j++;
-								}
-
-								if (!Mesh->GetComponentTransform().Equals(GetComponentTransform())) {
-									Mesh->SetRelativeTransform(GetComponentTransform());
-								}
-								Mesh->UpdateMeshSection_LinearColor(0, RenderMesh->Vertices, RenderMesh->Normals, RenderMesh->UVs, RenderMesh->Colors, RenderMesh->Tangents);
-							}
-
-							bCanRender = true;
-						}
-					}
-				}
-			}
-		}
-	}
-	if (!bCanRender)
-	{
-		ResetProceduralMesh();
-	}
-}
-
-
-void UFleshComponent::DebugDrawSkeletalMeshBindingPositions() const
-{
-#if WITH_EDITOR
-	auto UEVertd = [](FVector3f V) { return FVector3d(V.X, V.Y, V.Z); };
-	float SimulationBlendWeight = CVarParams.DrawSkeletalMeshBindingPositionsSimulationBlendWeight;
-
-	if (const UFleshAsset* RestAsset = GetRestCollection())
-	{
-		const USkeletalMesh* SkeletalMesh = 
-			GetRestCollection() && GetRestCollection()->TargetDeformationSkeleton ? 
-			GetRestCollection()->TargetDeformationSkeleton : RestAsset->SkeletalMesh;
-		if (SkeletalMesh)
-		{
-			TArray<bool> Influenced;
-			TArray<FVector> PosArray = GetSkeletalMeshEmbeddedPositionsInternal(
-				ChaosDeformableBindingOption::ComponentPos, FTransform::Identity, "", SimulationBlendWeight, &Influenced);
-
-			for (int i = 0; i < PosArray.Num(); i++)
-			{
-				const FVector& Pos = PosArray[i];
-				if (Influenced[i])
-				{
-					Chaos::FDebugDrawQueue::GetInstance().DrawDebugPoint(
-						GetComponentTransform().TransformPosition(Pos), FColor::Yellow, true, 2.0f, SDPG_Foreground, 10);
-				}
-			}
-		}
-	}
-#endif
-}
-
-TArray<FVector> UFleshComponent::GetSkeletalMeshBindingPositions(const USkeletalMesh* InSkeletalMesh) const
-{
-	return GetSkeletalMeshBindingPositionsInternal(InSkeletalMesh, nullptr);
-}
-
-TArray<FVector> UFleshComponent::GetSkeletalMeshEmbeddedPositions(
-	const ChaosDeformableBindingOption Format,
-	const FTransform TargetDeformationSkeletonOffset,
-	const FName TargetBone,
-	const float SimulationBlendWeight) const
-{
-	return GetSkeletalMeshEmbeddedPositionsInternal(Format, TargetDeformationSkeletonOffset, TargetBone, SimulationBlendWeight, nullptr);
 }
 
-TArray<FVector> UFleshComponent::GetSkeletalMeshEmbeddedPositionsInternal(
-	const ChaosDeformableBindingOption Format,
-	const FTransform TargetDeformationSkeletonOffset, 
-	const FName TargetBone,
-	const float SimulationBlendWeight,
-	TArray<bool>* OutInfluence) const
-{
-	if (!GetRestCollection())
-	{
-		UE_LOG(LogFleshComponentInternal, Warning,
-			TEXT("'%s' - GetSkeletalMeshEmbeddedPositionsInternal - RestCollection is not set."),
-			*GetName());
-		return TArray<FVector>();
-	}
-	if(!GetRestCollection()->TargetDeformationSkeleton)
-	{
-		UE_LOG(LogFleshComponentInternal, Warning,
-			TEXT("'%s' - GetSkeletalMeshEmbeddedPositionsInternal - TargetDeformationSkeleton is not set on the flesh asset."),
-			*GetName());
-		return TArray<FVector>();
-	}
-
-	TArray<FVector> EmbeddedPosComp;
-
-	// Get sample points in the skel mesh's component space.  This code assumes that the
-	// skeletal mesh and the flesh asset are aligned in their respective local spaces.  If
-	// they're not aligned, then the TargetDeformationSkeletonOffset should be provided to
-	// put the skeletal mesh in the same place as the flesh mesh.
-	TArray<FVector> TransformPositions;
-	if (Format == ChaosDeformableBindingOption::WorldDelta || 
-		Format == ChaosDeformableBindingOption::ComponentDelta) // BoneDelta handled below
-	{
-		TArray<FTransform> ComponentPose;
-		Dataflow::Animation::GlobalTransforms(GetRestCollection()->TargetDeformationSkeleton->GetRefSkeleton(), ComponentPose);
-		TransformPositions.SetNumUninitialized(ComponentPose.Num());
-		if (TargetDeformationSkeletonOffset.Equals(FTransform::Identity))
-		{
-			for (int32 i = 0; i < ComponentPose.Num(); i++)
-			{
-				TransformPositions[i] = ComponentPose[i].GetTranslation();
-			}
-		}
-		else
-		{
-			for (int32 i = 0; i < ComponentPose.Num(); i++)
-			{
-				TransformPositions[i] =
-					TargetDeformationSkeletonOffset.TransformPosition(
-						ComponentPose[i].GetTranslation());
-			}
-		}
-	}
-	else
-	{
-		// If we aren't computing deltas, then we don't need the bone positions.  We
-		// only need to size the array to however many bones we have.
-		TransformPositions.SetNum(GetRestCollection()->TargetDeformationSkeleton->GetRefSkeleton().GetNum());
-	}
-
-
-
-	// World space
-	if (Format == ChaosDeformableBindingOption::WorldPos)
-	{
-		// Calculate their current embedded positions
-		EmbeddedPosComp =
-			GetEmbeddedPositionsInternal(
-			TransformPositions,
-			FName(GetRestCollection()->TargetDeformationSkeleton->GetName()), // for identifying the binding group
-			SimulationBlendWeight,
-			OutInfluence);
-
-		// Put component space points into world space
-		const FTransform ComponentXf = GetComponentTransform();
-		for (int32 i = 0; i < EmbeddedPosComp.Num(); i++)
-		{
-			EmbeddedPosComp[i] = ComponentXf.TransformPosition(EmbeddedPosComp[i]);
-		}
-	}
-	else if (Format == ChaosDeformableBindingOption::WorldDelta)
-	{
-		// Calculate their current embedded positions
-		EmbeddedPosComp =
-			GetEmbeddedPositionsInternal(
-			TransformPositions,
-			FName(GetRestCollection()->TargetDeformationSkeleton->GetName()), // for identifying the binding group
-			1.f,
-			OutInfluence);
-
-		const FTransform ComponentXf = GetComponentTransform();
-		for (int32 i = 0; i < EmbeddedPosComp.Num(); i++)
-		{
-			EmbeddedPosComp[i] = EmbeddedPosComp[i] - TransformPositions[i];
-			EmbeddedPosComp[i] = ComponentXf.TransformVector(EmbeddedPosComp[i]);
-		}
-	}
-
-	// Component space
-	else if (Format == ChaosDeformableBindingOption::ComponentPos)
-	{
-		EmbeddedPosComp =
-			GetEmbeddedPositionsInternal(
-			TransformPositions,
-			FName(GetRestCollection()->TargetDeformationSkeleton->GetName()), // for identifying the binding group
-			1.f,
-			OutInfluence);
-	}
-	else if (Format == ChaosDeformableBindingOption::ComponentDelta)
-	{
-		// Calculate their current embedded positions
-		EmbeddedPosComp =
-			GetEmbeddedPositionsInternal(
-			TransformPositions,
-			FName(GetRestCollection()->TargetDeformationSkeleton->GetName()), // for identifying the binding group
-			1.f,
-			OutInfluence);
-
-		for (int32 i = 0; i < EmbeddedPosComp.Num(); i++)
-		{
-			EmbeddedPosComp[i] = (EmbeddedPosComp[i] - TransformPositions[i]) * SimulationBlendWeight;
-		}
-	}
-
-	// Bone space
-	else
-	{
-		EmbeddedPosComp =
-			GetEmbeddedPositionsInternal(
-			TransformPositions,
-			FName(GetRestCollection()->TargetDeformationSkeleton->GetName()), // for identifying the binding group
-			SimulationBlendWeight,
-			OutInfluence);
-
-
-		// Find the component that owns TargetDeformationSkeleton, so we can pull the 
-		// current animated bone transforms out of it.
-		//
-		// It's possible that TargetDeformationSkeleton is owned by a component on another 
-		// actor, in which case, we can't (easily) find it. If that becomes a desired use
-		// case, then we'll need to have a handle to that component, not just the asset.
-		TArray<USkeletalMeshComponent*> SkeletalMeshComponents;
-		GetOwner()->GetComponents<USkeletalMeshComponent>(SkeletalMeshComponents);
-		USkeletalMeshComponent* TargetDeformationSkeletalMeshComponent = nullptr;
-		for (int32 i = 0; i < SkeletalMeshComponents.Num(); i++)
-		{
-			if (SkeletalMeshComponents[i]->GetSkeletalMeshAsset() == GetRestCollection()->TargetDeformationSkeleton)
-			{
-				TargetDeformationSkeletalMeshComponent = SkeletalMeshComponents[i];
-				break;
-			}
-		}
-
-		if (TargetDeformationSkeletalMeshComponent)
-		{
-			// Get the current animated bone transforms and then their positions.
-			const TArray<FTransform>& TargetDeformationComponentTransforms = 
-				TargetDeformationSkeletalMeshComponent->GetComponentSpaceTransforms();
-			TArray<FVector> AnimTransformPositions;
-			AnimTransformPositions.SetNumUninitialized(TargetDeformationComponentTransforms.Num());
-
-			// Get their positions, and apply the offset if necessary.
-			if (!TargetDeformationSkeletonOffset.Equals(FTransform::Identity) && 
-				Format == ChaosDeformableBindingOption::BoneDelta)
-			{
-				for (int32 i = 0; i < TargetDeformationComponentTransforms.Num(); i++)
-				{
-					AnimTransformPositions[i] =
-						TargetDeformationSkeletonOffset.TransformPosition(
-							TargetDeformationComponentTransforms[i].GetTranslation());
-				}
-			}
-			else
-			{
-				for (int32 i = 0; i < TargetDeformationComponentTransforms.Num(); i++)
-				{
-					AnimTransformPositions[i] = TargetDeformationComponentTransforms[i].GetTranslation();
-				}
-			}
-
-			// Find the transform index of 'TargetBone'
-			FTransform BoneToComponentXf = FTransform::Identity;
-			const int32 BoneIndex = GetRestCollection()->TargetDeformationSkeleton->GetRefSkeleton().FindBoneIndex(TargetBone);
-			if (TargetDeformationComponentTransforms.IsValidIndex(BoneIndex))
-			{
-				BoneToComponentXf = FTransform(TargetDeformationComponentTransforms[BoneIndex].ToMatrixWithScale().Inverse());
-			}
-			else
-			{
-				UE_LOG(LogFleshComponentInternal, Warning, 
-					TEXT("'%s' - Failed to find a valid bone index (got %d) for bone name '%s' in TargetDeformationSkeleton '%s' "
-					"corresponding to SkeletalMeshComponent '%s', which has %d bones."), 
-					*GetName(),
-					BoneIndex, *TargetBone.ToString(), *GetRestCollection()->TargetDeformationSkeleton.GetName(),
-					*TargetDeformationSkeletalMeshComponent->GetName(), TargetDeformationComponentTransforms.Num());
-				return TArray<FVector>();
-			}
-
-			// Compute the return values
-			if (Format == ChaosDeformableBindingOption::BonePos)
-			{
-				if (!BoneToComponentXf.Equals(FTransform::Identity))
-				{
-					for (int32 i = 0; i < EmbeddedPosComp.Num(); i++)
-					{
-						EmbeddedPosComp[i] = BoneToComponentXf.TransformPosition(EmbeddedPosComp[i]);
-					}
-				}
-			}
-			else if (Format == ChaosDeformableBindingOption::BoneDelta)
-			{
-				if (!BoneToComponentXf.Equals(FTransform::Identity))
-				{
-					for (int32 i = 0; i < EmbeddedPosComp.Num(); i++)
-					{
-						EmbeddedPosComp[i] = EmbeddedPosComp[i] - AnimTransformPositions[i];
-						EmbeddedPosComp[i] = BoneToComponentXf.TransformVector(EmbeddedPosComp[i]);
-					}
-				}
-				else
-				{
-					for (int32 i = 0; i < EmbeddedPosComp.Num(); i++)
-					{
-						EmbeddedPosComp[i] = EmbeddedPosComp[i] - AnimTransformPositions[i];
-					}
-				}
-			}
-		}
-		else
-		{
-			UE_LOG(LogFleshComponentInternal, Warning, 
-				TEXT("'%s' - Failed to find SkeletalMeshComponent for TargetDeformationSkeleton '%s'."),
-				*GetName(),
-				*GetRestCollection()->TargetDeformationSkeleton->GetName());
-			return TArray<FVector>();
-		}
-	}
-	return EmbeddedPosComp;
-}
-
-TArray<FVector> UFleshComponent::GetEmbeddedPositionsInternal(
-	const TArray<FVector>& InPositions, 
-	const FName SkeletalMeshName,
-	const float SimulationBlendWeight,
-	TArray<bool>* OutInfluence) const
-{
-	auto UEVert3d = [](FVector3f V) { return FVector3d(V.X, V.Y, V.Z); };
-	auto UEVert4d = [](FVector4f V) { return FVector4d(V.X, V.Y, V.Z, V.W); };
-
-	TArray<FVector> OutPositions = InPositions;
-
-	if (const UFleshAsset* RestAsset = GetRestCollection())
-	{
-		if (const FFleshCollection* Rest = RestAsset->GetCollection())
-		{
-			GeometryCollection::Facades::FTetrahedralSkeletalBindings TetBindings(*Rest);
-
-			const TManagedArray<int32>* TetrahedronStart = Rest->FindAttribute<int32>(FTetrahedralCollection::TetrahedronStartAttribute, FGeometryCollection::GeometryGroup);
-			const TManagedArray<FVector3f>* RestVerts = Rest->FindAttributeTyped<FVector3f>("Vertex", FGeometryCollection::VerticesGroup);
-			const TManagedArray<FVector3f>* Verts = GetDynamicCollection() ? GetDynamicCollection()->FindPositions() : RestVerts;
-
-			if (ensure(Verts != nullptr) && TetrahedronStart)
-			{
-				auto CalculateBindings = [this, TetBindings, TetrahedronStart, SkeletalMeshName, OutInfluence](const TManagedArray<FVector3f>* Verts, TArray<FVector>& OutPositions)
-				{
-					if (OutInfluence != nullptr)
-						OutInfluence->Init(false, OutPositions.Num());
-					for (int32 TetMeshIdx = 0; TetMeshIdx < TetrahedronStart->Num(); TetMeshIdx++)
-					{
-						FString MeshBindingsName =
-							GeometryCollection::Facades::FTetrahedralSkeletalBindings::GenerateMeshGroupName(TetMeshIdx, SkeletalMeshName);
-						if (!TetBindings.CalculateBindings(MeshBindingsName, Verts->GetConstArray(), OutPositions, OutInfluence))
-						{
-							UE_LOG(LogFleshComponentInternal, Warning,
-								   TEXT("'%s' - GetEmbeddedPositionsInternal - Failed to find mesh bindings for skeletal mesh '%s'"),
-								   *GetName(),
-								   *SkeletalMeshName.ToString());
-						}
-					}
-				};
-				CalculateBindings(Verts, OutPositions);
-
-				// blend between the aniamtion position and the 
-				// simulated position. 
-				if (!FMath::IsNearlyEqual(SimulationBlendWeight, 1.0) && GetDynamicCollection())
-				{
-					TArray<FVector> RestPositions = InPositions;
-					CalculateBindings(RestVerts, RestPositions);
-
-					float ClampedWeight = FMath::Clamp(SimulationBlendWeight, 0.f, 1.f);
-					for (int i = 0; i < OutPositions.Num(); i++)
-					{
-						FVector V = RestPositions[i] - OutPositions[i];
-						OutPositions[i] += V * (1.0-ClampedWeight);
-					}
-				}
-			}
-		}
-	}
-	return OutPositions;
-}
-
-TArray<FVector> UFleshComponent::GetSkeletalMeshBindingPositionsInternal(const USkeletalMesh* InSkeletalMesh, TArray<bool>* OutInfluence) const
-{
-	auto UEVert3d = [](FVector3f V) { return FVector3d(V.X, V.Y, V.Z); };
-	auto UEVert4d = [](FVector4f V) { return FVector4d(V.X, V.Y, V.Z, V.W); };
-
-	TArray<FVector> TransformPositions;
-	if (InSkeletalMesh)
-	{
-		FName SkeletalMeshName(InSkeletalMesh->GetName());
-		if (const UFleshAsset* RestAsset = GetRestCollection())
-		{
-			if (const FFleshCollection* Rest = RestAsset->GetCollection())
-			{
-				GeometryCollection::Facades::FTetrahedralSkeletalBindings TetBindings(*Rest);
-
-				const TManagedArray<int32>* TetrahedronStart = Rest->FindAttribute<int32>(FTetrahedralCollection::TetrahedronStartAttribute, FGeometryCollection::GeometryGroup);
-				const TManagedArray<FVector3f>* Verts = GetDynamicCollection() ? GetDynamicCollection()->FindPositions() : Rest->FindAttributeTyped<FVector3f>("Vertex", FGeometryCollection::VerticesGroup);
-
-				if (ensure(Verts != nullptr) && TetrahedronStart)
-				{
-					// Component relative transforms, not world.
-					TArray<FTransform> ComponentPose;
-					Dataflow::Animation::GlobalTransforms(InSkeletalMesh->GetRefSkeleton(), ComponentPose);
-
-					TransformPositions.SetNumUninitialized(ComponentPose.Num());
-					for (int32 i = 0; i < ComponentPose.Num(); i++)
-					{
-						TransformPositions[i] = ComponentPose[i].GetTranslation();
-					}
-
-					if (OutInfluence != nullptr) OutInfluence->Init(false, TransformPositions.Num());
-					for (int32 TetMeshIdx = 0; TetMeshIdx < TetrahedronStart->Num(); TetMeshIdx++)
-					{
-						FString MeshBindingsName = GeometryCollection::Facades::FTetrahedralSkeletalBindings::GenerateMeshGroupName(TetMeshIdx, SkeletalMeshName);
-						TetBindings.CalculateBindings(MeshBindingsName, Verts->GetConstArray(), TransformPositions, OutInfluence);
-					}
-				}
-			}
-		}
-	}
-	return TransformPositions;
-}
-
-
diff --git a/Engine/Plugins/Experimental/ChaosFlesh/Source/ChaosFleshEngine/Public/ChaosFlesh/ChaosDeformableCollisionsActor.h b/Engine/Plugins/Experimental/ChaosFlesh/Source/ChaosFleshEngine/Public/ChaosFlesh/ChaosDeformableCollisionsActor.h
--- a/Engine/Plugins/Experimental/ChaosFlesh/Source/ChaosFleshEngine/Public/ChaosFlesh/ChaosDeformableCollisionsActor.h
+++ b/Engine/Plugins/Experimental/ChaosFlesh/Source/ChaosFleshEngine/Public/ChaosFlesh/ChaosDeformableCollisionsActor.h
@@ -29,10 +29,10 @@ public:
 	TObjectPtr<UDeformableCollisionsComponent> DeformableCollisionsComponent;
 	UDeformableCollisionsComponent* GetCollisionsComponent() const { return DeformableCollisionsComponent; }
 
-	UPROPERTY(EditAnywhere, BlueprintReadWrite, Category = "Physics")
+	UPROPERTY(EditAnywhere, BlueprintReadWrite, Category = "Physics", meta = (DisplayPriority = 1))
 	TObjectPtr<ADeformableSolverActor> PrimarySolver;
 
-	UPROPERTY(EditAnywhere, BlueprintReadWrite, Category = "Physics")
+	UPROPERTY(EditAnywhere, BlueprintReadWrite, Category = "Physics", meta = (DisplayPriority = 2))
 	TArray<TObjectPtr<AStaticMeshActor>> StaticCollisions;
 
 #if WITH_EDITOR
diff --git a/Engine/Plugins/Experimental/ChaosFlesh/Source/ChaosFleshEngine/Public/ChaosFlesh/ChaosDeformableGameplayComponent.h b/Engine/Plugins/Experimental/ChaosFlesh/Source/ChaosFleshEngine/Public/ChaosFlesh/ChaosDeformableGameplayComponent.h
--- a/Engine/Plugins/Experimental/ChaosFlesh/Source/ChaosFleshEngine/Public/ChaosFlesh/ChaosDeformableGameplayComponent.h
+++ b/Engine/Plugins/Experimental/ChaosFlesh/Source/ChaosFleshEngine/Public/ChaosFlesh/ChaosDeformableGameplayComponent.h
@@ -3,7 +3,7 @@
 
 #include "CoreMinimal.h"
 
-#include "ChaosFlesh/FleshComponent.h"
+#include "ChaosFlesh/ChaosDeformableTetrahedralComponent.h"
 #include "UObject/ObjectMacros.h"
 
 
@@ -17,27 +17,42 @@ struct FRigBoundRayCasts
 {
 	GENERATED_USTRUCT_BODY()
 
-	UPROPERTY(EditAnywhere, Category = "RigBoundRaycasts")
+	UPROPERTY(EditAnywhere, Category = "Physics")
 		bool bEnableRigBoundRaycasts = false;
 
-	UPROPERTY(EditAnywhere, Category = "RigBoundRaycasts")
+	UPROPERTY(EditAnywhere, Category = "Physics")
 		int32 MaxNumTests = 5;
 
-	UPROPERTY(EditAnywhere, Category = "RigBoundRaycasts")
+	UPROPERTY(EditAnywhere, Category = "Physics")
 		bool bTestDownOnly = false;
 
-	UPROPERTY(EditAnywhere, Category = "RigBoundRaycasts")
+	UPROPERTY(EditAnywhere, Category = "Physics")
 		float TestRange = 0.5f;
+	
+	/** Objects to skip if hit during \c DetectEnvironmentCollisions(). */
+	UPROPERTY(EditAnywhere, Category = "Physics")
+	TArray<TWeakObjectPtr<UPrimitiveComponent>> EnvironmentCollisionsSkipList;
 
-	UPROPERTY(EditAnywhere, Category = "RigBoundRaycasts")
+	UPROPERTY(EditAnywhere, Category = "Physics")
 		TEnumAsByte<ECollisionChannel> CollisionChannel = ECC_WorldStatic;
 };
 
+
+USTRUCT()
+struct FGameplayColllisions
+{
+	GENERATED_USTRUCT_BODY()
+
+	UPROPERTY(EditAnywhere, Category = "Physics")
+	FRigBoundRayCasts RigBoundRayCasts;
+
+};
+
 /**
 *	UDeformableGameplayComponent
 */
 UCLASS(meta = (BlueprintSpawnableComponent))
-class CHAOSFLESHENGINE_API UDeformableGameplayComponent : public UFleshComponent
+class CHAOSFLESHENGINE_API UDeformableGameplayComponent : public UDeformableTetrahedralComponent
 {
 	GENERATED_UCLASS_BODY()
 
@@ -67,12 +82,10 @@ public:
 	void DetectEnvironmentCollisions(const int32 MaxNumTests = 100, const bool bTestDownOnly = true, const float TestRange = 0.0, const ECollisionChannel CollisionChannel = ECollisionChannel::ECC_PhysicsBody);
 
 
-	UPROPERTY(EditAnywhere, Category = "Physics")
-	FRigBoundRayCasts RigBoundRayCasts;
+	UPROPERTY(EditAnywhere, Category = "Physics", DisplayName = "Collisions", meta = (DisplayPriority = 10))
+	FGameplayColllisions GameplayColllisions;
+
 
-	/** Objects to skip if hit during \c DetectEnvironmentCollisions(). */
-	UPROPERTY(EditAnywhere, Category = "Physics")
-	TArray<TWeakObjectPtr<UPrimitiveComponent>> EnvironmentCollisionsSkipList;
 
 private:
 	TSet<int32> EnvCollisionsPreviousHits;
diff --git a/Engine/Plugins/Experimental/ChaosFlesh/Source/ChaosFleshEngine/Public/ChaosFlesh/ChaosDeformablePhysicsComponent.h b/Engine/Plugins/Experimental/ChaosFlesh/Source/ChaosFleshEngine/Public/ChaosFlesh/ChaosDeformablePhysicsComponent.h
--- a/Engine/Plugins/Experimental/ChaosFlesh/Source/ChaosFleshEngine/Public/ChaosFlesh/ChaosDeformablePhysicsComponent.h
+++ b/Engine/Plugins/Experimental/ChaosFlesh/Source/ChaosFleshEngine/Public/ChaosFlesh/ChaosDeformablePhysicsComponent.h
@@ -13,6 +13,7 @@
 class ADeformableSolverActor;
 class UDeformableSolverComponent;
 
+
 /**
 *	UDeformablePhysicsComponent
 */
@@ -51,26 +52,10 @@ public:
 	UDeformableSolverComponent* GetDeformableSolver();
 	const UDeformableSolverComponent* GetDeformableSolver() const;
 	
-	UPROPERTY(EditAnywhere, BlueprintReadWrite, Category = "Physics")
-	TObjectPtr<UDeformableSolverComponent> PrimarySolverComponent;
-
-	UPROPERTY(EditAnywhere, BlueprintReadWrite, Category = "Physics")
-	bool bTempEnableGravity = true;
-
-	UPROPERTY(EditAnywhere, BlueprintReadWrite, Category = "Physics", meta = (ClampMin = "0.0", ClampMax = "1.0", UIMin = "0.0", UIMax = "1.0"))
-	float DampingMultiplier = 1.f;
 
-	UPROPERTY(EditAnywhere, BlueprintReadWrite, Category = "Physics", meta = (ClampMin = "0.0", ClampMax = "1.0", UIMin = "0.0", UIMax = "1.0"))
-	float StiffnessMultiplier = 1.f;
-
-	UPROPERTY(EditAnywhere, BlueprintReadWrite, Category = "Physics", meta = (ClampMin = "0.0", ClampMax = "1.0", UIMin = "0.0", UIMax = "1.0"))
-	float MassMultiplier = 1.f;
-
-	UPROPERTY(EditAnywhere, BlueprintReadWrite, Category = "Physics", meta = (ClampMin = "0.0", ClampMax = "1.0", UIMin = "0.0", UIMax = "1.0"))
-	float IncompressibilityMultiplier = 1.f;
+	UPROPERTY(EditAnywhere, BlueprintReadWrite, Category = "Physics", meta = (DisplayPriority = 2))
+	TObjectPtr<UDeformableSolverComponent> PrimarySolverComponent;
 
-	UPROPERTY(EditAnywhere, BlueprintReadWrite, Category = "Physics", meta = (ClampMin = "0.0", ClampMax = "1.0", UIMin = "0.0", UIMax = "1.0"))
-	float InflationMultiplier = 1.f;
 
 	const FThreadingProxy* GetPhysicsProxy() const { return PhysicsProxy; }
 	FThreadingProxy* GetPhysicsProxy() { return PhysicsProxy; }
diff --git a/Engine/Plugins/Experimental/ChaosFlesh/Source/ChaosFleshEngine/Public/ChaosFlesh/ChaosDeformableSolverComponent.h b/Engine/Plugins/Experimental/ChaosFlesh/Source/ChaosFleshEngine/Public/ChaosFlesh/ChaosDeformableSolverComponent.h
--- a/Engine/Plugins/Experimental/ChaosFlesh/Source/ChaosFleshEngine/Public/ChaosFlesh/ChaosDeformableSolverComponent.h
+++ b/Engine/Plugins/Experimental/ChaosFlesh/Source/ChaosFleshEngine/Public/ChaosFlesh/ChaosDeformableSolverComponent.h
@@ -16,6 +16,146 @@
 class UDeformablePhysicsComponent;
 class UDeformableCollisionsComponent;
 
+USTRUCT(BlueprintType)
+struct FConnectedObjectsGroup
+{
+	GENERATED_USTRUCT_BODY()
+
+	UPROPERTY(EditAnywhere, BlueprintReadOnly, Category = "ConnectedObjects", meta = (EditCondition = "false"))
+	TArray< TObjectPtr<UDeformablePhysicsComponent> > DeformableComponents;
+};
+
+USTRUCT(BlueprintType)
+struct FSolverTimingGroup
+{
+	GENERATED_USTRUCT_BODY()
+
+	UPROPERTY(EditAnywhere, Category = "SolverTiming")
+		int32 NumSubSteps = 2;
+
+	UPROPERTY(EditAnywhere, Category = "SolverTiming")
+		int32 NumSolverIterations = 5;
+
+	UPROPERTY(EditAnywhere, Category = "SolverTiming")
+		bool FixTimeStep = false;
+
+	UPROPERTY(EditAnywhere, Category = "SolverTiming")
+		float TimeStepSize = 0.05;
+
+	UPROPERTY(EditAnywhere, Category = "SolverTiming")
+		bool bDoThreadedAdvance = true;
+
+	/** ObjectType defines how to initialize the rigid objects state, Kinematic, Sleeping, Dynamic. */
+	UPROPERTY(EditAnywhere, BlueprintReadOnly, Category = "SolverTiming")
+		EDeformableExecutionModel ExecutionModel = EDeformableExecutionModel::Chaos_Deformable_PostPhysics;
+};
+
+
+USTRUCT(BlueprintType)
+struct FSolverDebuggingGroup
+{
+	GENERATED_USTRUCT_BODY()
+
+	UPROPERTY(EditAnywhere, Category = "Physics")
+	bool CacheToFile = false;
+};
+
+
+USTRUCT(BlueprintType)
+struct FSolverQuasistaticsGroup
+{
+	GENERATED_USTRUCT_BODY()
+
+	UPROPERTY(EditAnywhere, Category = "Quasistatics")
+	bool bDoQuasistatics = false;
+};
+
+
+USTRUCT(BlueprintType)
+struct FSolverEvolutionGroup
+{
+	GENERATED_USTRUCT_BODY()
+
+	UPROPERTY(EditAnywhere, Category = "Evolution")
+	FSolverQuasistaticsGroup SolverQuasistatics;
+
+};
+
+
+USTRUCT(BlueprintType)
+struct FSolverGridBasedCollisionsGroup
+{
+	GENERATED_USTRUCT_BODY()
+
+	UPROPERTY(EditAnywhere, Category = "GridBasedCollisions")
+		bool bUseGridBasedConstraints = false;
+
+	UPROPERTY(EditAnywhere, Category = "GridBasedCollisions")
+		float GridDx = 25.;
+};
+
+
+USTRUCT(BlueprintType)
+struct FSolverCollisionsGroup
+{
+	GENERATED_USTRUCT_BODY()
+
+	UPROPERTY(EditAnywhere, Category = "Collisions")
+	bool bDoSelfCollision = false;
+
+	UPROPERTY(EditAnywhere, Category = "Physics")
+	bool bUseFloor = true;
+
+	//UPROPERTY(EditAnywhere, Category = "Collisions")
+	//FSolverGridBasedCollisionsGroup SolverGridBasedCollisions;
+};
+
+USTRUCT(BlueprintType)
+struct FSolverCorotatedConstraintsGroup
+{
+	GENERATED_USTRUCT_BODY()
+
+	UPROPERTY(EditAnywhere, Category = "Corotated")
+	bool bEnableCorotatedConstraint = true;
+
+	UPROPERTY(EditAnywhere, Category = "Corotated")
+	bool bDoBlended = false;
+
+	UPROPERTY(EditAnywhere, Category = "Corotated")
+	float BlendedZeta = 0;
+};
+
+USTRUCT(BlueprintType)
+struct FSolverConstraintsGroup
+{
+	GENERATED_USTRUCT_BODY()
+
+	UPROPERTY(EditAnywhere, Category = "Constraints")
+	bool bEnablePositionTargets = true;
+
+	UPROPERTY(EditAnywhere, Category = "Constraints")
+	bool bEnableKinematics = true;
+
+	UPROPERTY(EditAnywhere, Category = "Constraints")
+	FSolverCorotatedConstraintsGroup CorotatedConstraints;
+};
+
+
+USTRUCT(BlueprintType)
+struct FSolverForcesGroup
+{
+	GENERATED_USTRUCT_BODY()
+
+	UPROPERTY(EditAnywhere, Category = "Forces")
+	float YoungModulus = 100000;
+
+	UPROPERTY(EditAnywhere, Category = "Forces")
+	float Damping = 0;
+
+	UPROPERTY(EditAnywhere, Category = "Forces")
+	bool bEnableGravity = true;
+};
+
 /**
 *	UDeformableSolverComponent
 */
@@ -52,76 +192,27 @@
 	//bool ShouldWaitForDeformableInTickFunction() const;
 	void UpdateDeformableEndTickState(bool bRegister);
 
-	/* Properties */
+	/* Properties : Do NOT place ungrouped properties in this class */
+	UPROPERTY(EditAnywhere, BlueprintReadOnly, Category = "Physics", meta = (EditCondition = "false"))
+	FConnectedObjectsGroup ConnectedObjects;
 
-	UPROPERTY(EditAnywhere, BlueprintReadOnly, Category = "Physics", meta = (EditCondition = "false"))
-	TArray< TObjectPtr<UDeformablePhysicsComponent> > DeformableComponents;
+	UPROPERTY(EditAnywhere, BlueprintReadOnly, Category = "Physics")
+	FSolverTimingGroup SolverTiming;
 
-	UPROPERTY(EditAnywhere, BlueprintReadOnly, Category = "Physics", meta = (EditCondition = "false"))
-	TObjectPtr<UDeformableCollisionsComponent> CollisionComponent;
+	UPROPERTY(EditAnywhere, BlueprintReadOnly, Category = "Physics")
+	FSolverEvolutionGroup SolverEvolution;
 
-	/** ObjectType defines how to initialize the rigid objects state, Kinematic, Sleeping, Dynamic. */
-	UPROPERTY(EditAnywhere, BlueprintReadOnly , Category = "Physics")
-	EDeformableExecutionModel ExecutionModel = EDeformableExecutionModel::Chaos_Deformable_DuringPhysics;
+	UPROPERTY(EditAnywhere, BlueprintReadOnly, Category = "Physics")
+	FSolverCollisionsGroup SolverCollisions;
 
-	UPROPERTY(EditAnywhere, Category = "Physics")
-		int32 NumSubSteps = 2;
+	UPROPERTY(EditAnywhere, BlueprintReadOnly, Category = "Physics")
+	FSolverConstraintsGroup SolverConstraints;
 
-	UPROPERTY(EditAnywhere, Category = "Physics")
-		int32 NumSolverIterations = 5;
+	UPROPERTY(EditAnywhere, BlueprintReadOnly, Category = "Physics")
+	FSolverForcesGroup SolverForces;
 
-	UPROPERTY(EditAnywhere, Category = "Physics")
-		bool FixTimeStep = false;
+	UPROPERTY(EditAnywhere, BlueprintReadOnly, Category = "Physics")
+	FSolverDebuggingGroup SolverDebugging;
 
-	UPROPERTY(EditAnywhere, Category = "Physics")
-		float TimeStepSize = 0.05;
-
-	UPROPERTY(EditAnywhere, Category = "Physics")
-		bool CacheToFile = false;
-
-	UPROPERTY(EditAnywhere, Category = "Physics")
-		bool bEnableKinematics = true;
-
-	UPROPERTY(EditAnywhere, Category = "Physics")
-		bool bUseFloor = true;
-
-	UPROPERTY(EditAnywhere, Category = "Physics")
-		bool bDoSelfCollision = false;
-
-	UPROPERTY(EditAnywhere, Category = "Physics")
-		bool bDoThreadedAdvance = true;
-
-	UPROPERTY(EditAnywhere, Category = "Physics|Expermental" )
-		bool bUseGridBasedConstraints = false;
-
-	UPROPERTY(EditAnywhere, Category = "Physics")
-		float GridDx = 25.;
-
-	UPROPERTY(EditAnywhere, Category = "Physics")
-		bool bDoQuasistatics = false;
-
-	UPROPERTY(EditAnywhere, Category = "Physics")
-		float YoungModulus = 100000;
-
-	UPROPERTY(EditAnywhere, Category = "Physics")
-		bool bDoBlended = false;
-
-	UPROPERTY(EditAnywhere, Category = "Physics")
-		float BlendedZeta = 0;
-
-	UPROPERTY(EditAnywhere, Category = "Physics")
-		float Damping = 0;
-
-	UPROPERTY(EditAnywhere, Category = "Physics")
-		bool bEnableGravity = true;
-
-	UPROPERTY(EditAnywhere, Category = "Physics")
-		bool bEnableCorotatedConstraint = true;
-
-	UPROPERTY(EditAnywhere, Category = "Physics")
-		bool bEnablePositionTargets = true;
-	//UPROPERTY(EditAnywhere, Category = Chaos)
-	//	bool bWaitForParallelDeformableTask = true;
-
 	// Simulation Variables
 	TUniquePtr<FDeformableSolver> Solver;
diff --git a/Engine/Plugins/Experimental/ChaosFlesh/Source/ChaosFleshEngine/Public/ChaosFlesh/ChaosDeformableTetrahedralComponent.h b/Engine/Plugins/Experimental/ChaosFlesh/Source/ChaosFleshEngine/Public/ChaosFlesh/ChaosDeformableTetrahedralComponent.h
new file mode 100644
--- /dev/null
+++ b/Engine/Plugins/Experimental/ChaosFlesh/Source/ChaosFleshEngine/Public/ChaosFlesh/ChaosDeformableTetrahedralComponent.h
@@ -0,0 +1,204 @@
+// Copyright Epic Games, Inc. All Rights Reserved.
+#pragma once
+
+#include "Chaos/Deformable/ChaosDeformableSolverProxy.h"
+#include "Chaos/Deformable/ChaosDeformableSolver.h"
+#include "ChaosFlesh/ChaosDeformableSolverThreading.h"
+#include "ChaosFlesh/ChaosDeformablePhysicsComponent.h"
+#include "ChaosFlesh/FleshAsset.h"
+#include "ChaosFlesh/FleshDynamicAsset.h"
+#include "ChaosFlesh/SimulationAsset.h"
+#include "Components/MeshComponent.h"
+#include "UObject/ObjectMacros.h"
+#include "ProceduralMeshComponent.h"
+#include "ChaosDeformableTetrahedralComponent.generated.h"
+
+class FFleshCollection;
+class ADeformableSolverActor;
+class UDeformableSolverComponent;
+
+/**
+*  Options for binding positions query.
+*/
+UENUM()
+enum ChaosDeformableBindingOption : uint8
+{
+	WorldPos		UMETA(DisplayName = "World Positions"),
+	WorldDelta		UMETA(DisplayName = "World Deltas"),
+	ComponentPos    UMETA(DisplayName = "Component Positions"),
+	ComponentDelta  UMETA(DisplayName = "Component Deltas"),
+	BonePos			UMETA(DisplayName = "Bone Positions"),
+	BoneDelta		UMETA(DisplayName = "Bone Deltas"),
+};
+
+
+USTRUCT(BlueprintType)
+struct FFleshSimulationSpaceGroup
+{
+	GENERATED_USTRUCT_BODY()
+
+		/**
+		* Bone from the associated skeletal mesh (indicated by RestCollection.TargetSkeletalMesh) to use as
+		* the space the sim runs in.
+		*/
+	UPROPERTY(EditAnywhere, Category = "Physics", meta = (GetOptions = "GetSimSpaceBoneNameOptions", EditCondition = "SimSpace == ChaosDeformableSimSpace::Bone"))
+	FName SimSpaceBoneName;
+
+
+	/** Space the simulation will run in. */
+	UPROPERTY(EditAnywhere, Category = "Physics")
+		TEnumAsByte<ChaosDeformableSimSpace> SimSpace = ChaosDeformableSimSpace::World;
+
+	/** The skeletal mesh to use pull the \c SimSpaceBoneName from. */
+	UPROPERTY()
+	TObjectPtr<USkeletalMesh> SimSpaceSkeletalMesh;
+
+
+	int32 SimSpaceTransformIndex = INDEX_NONE;
+	int32 SimSpaceTransformGlobalIndex = INDEX_NONE;
+};
+
+
+USTRUCT(BlueprintType)
+struct FBodyForcesGroup
+{
+	GENERATED_USTRUCT_BODY()
+
+	UPROPERTY(EditAnywhere, BlueprintReadWrite, Category = "Physics")
+	bool bApplyGravity = true;
+
+	UPROPERTY(EditAnywhere, BlueprintReadWrite, Category = "Physics", meta = (ClampMin = "0.0", ClampMax = "1.0", UIMin = "0.0", UIMax = "1.0"))
+		float DampingMultiplier = 1.f;
+
+	UPROPERTY(EditAnywhere, BlueprintReadWrite, Category = "Physics", meta = (ClampMin = "0.0", ClampMax = "1.0", UIMin = "0.0", UIMax = "1.0"))
+		float StiffnessMultiplier = 1.f;
+
+	UPROPERTY(EditAnywhere, BlueprintReadWrite, Category = "Physics", meta = (ClampMin = "0.0", ClampMax = "1.0", UIMin = "0.0", UIMax = "1.0"))
+		float IncompressibilityMultiplier = 1.f;
+
+	UPROPERTY(EditAnywhere, BlueprintReadWrite, Category = "Physics", meta = (ClampMin = "0.0", ClampMax = "1.0", UIMin = "0.0", UIMax = "1.0"))
+		float InflationMultiplier = 1.f;
+};
+
+/**
+*	UDeformableTetrahedralComponent
+*/
+UCLASS(meta = (BlueprintSpawnableComponent))
+class CHAOSFLESHENGINE_API UDeformableTetrahedralComponent : public UDeformablePhysicsComponent
+{
+	GENERATED_UCLASS_BODY()
+
+public:
+	typedef Chaos::Softs::FFleshThreadingProxy FFleshThreadingProxy;
+
+	~UDeformableTetrahedralComponent();
+
+	/** USceneComponent Interface */
+	virtual void BeginPlay() override;
+	virtual void OnRegister() override;
+	virtual void TickComponent(float DeltaTime, enum ELevelTick TickType, FActorComponentTickFunction *ThisTickFunction) override;
+	virtual FBoxSphereBounds CalcBounds(const FTransform& LocalToWorld) const override;
+	void UpdateLocalBounds();
+	virtual void EndPlay(const EEndPlayReason::Type ReasonEnd) override;
+	void Invalidate();
+
+	/** Simulation Interface*/
+	virtual FThreadingProxy* NewProxy() override;
+	virtual FDataMapValue NewDeformableData() override;
+	virtual void UpdateFromSimualtion(const FDataMapValue* SimualtionBuffer) override;
+
+	/** RestCollection */
+	void SetRestCollection(const UFleshAsset * InRestCollection);
+	const UFleshAsset* GetRestCollection() const { return RestCollection; }
+
+	/** DynamicCollection */
+	void ResetDynamicCollection();
+	UFleshDynamicAsset* GetDynamicCollection() { return DynamicCollection; }
+	const UFleshDynamicAsset* GetDynamicCollection() const { return DynamicCollection; }
+
+	/** SimulationCollection */
+	void ResetSimulationCollection();
+	USimulationAsset* GetSimulationCollection() { return SimulationCollection; }
+	const USimulationAsset* GetSimulationCollection() const { return SimulationCollection; }
+
+	/** @deprecated Use GetSkeletalMeshEmbeddedPositions() instead. */
+	UFUNCTION(BlueprintCallable, Category = "Physics", meta = (DeprecatedFunction, DeprecationMessage = "Use GetSkeletalMeshEmbeddedPositions() instead."))
+	TArray<FVector> GetSkeletalMeshBindingPositions(const USkeletalMesh* InSkeletalMesh) const;
+
+	/**
+	* Get the current positions of the transformation hierarchy from \c TargetDeformationSkeleton,
+	* deformed by the tetrahedral mesh.  Results can be in world space postions/deltas, component space
+	* positions/deltas, or bone space positions/deltas.  If a bone space is desired \p TargetBone
+	* must indicate which bone to use. TargetDeformationSkeletonOffset is an offset transform that moves 
+	* the \c TargetDeformationSkeleton to be co-located with the flesh mesh.
+	*/
+	UFUNCTION(BlueprintCallable, Category = "Physics")
+	TArray<FVector> GetSkeletalMeshEmbeddedPositions(const ChaosDeformableBindingOption Format, const FTransform TargetDeformationSkeletonOffset, const FName TargetBone = "", const float SimulationBlendWeight = 1.f) const;
+
+	UPROPERTY(EditAnywhere, Category = "Physics", meta = (DisplayPriority = 10))
+	FFleshSimulationSpaceGroup SimulationSpace;
+
+	UPROPERTY(EditAnywhere, BlueprintReadWrite, Category = "Physics", meta = (DisplayPriority = 10))
+	FBodyForcesGroup BodyForces;
+
+	UPROPERTY(EditAnywhere, BlueprintReadWrite, Category = "Physics", meta = (ClampMin = "0.0", ClampMax = "1.0", UIMin = "0.0", UIMax = "1.0", DisplayPriority = 100))
+	float MassMultiplier = 1.f;
+
+	UPROPERTY(EditAnywhere, Category = "Rendering")
+	TObjectPtr<UProceduralMeshComponent> Mesh;
+
+private:
+	/** FleshAsset that describes the simulation rest state. */
+	UPROPERTY(EditAnywhere, Category = "Physics", meta = (DisplayPriority = 1))
+	TObjectPtr<const UFleshAsset> RestCollection;
+
+	/** Current simulation state. */
+	UPROPERTY()
+	TObjectPtr<UFleshDynamicAsset> DynamicCollection;
+
+	/** Simulator input */
+	UPROPERTY()
+	TObjectPtr<USimulationAsset> SimulationCollection;
+
+
+	/* Returns a list of bone names from the currently selected skeletal mesh. */
+	UFUNCTION(CallInEditor)
+	TArray<FString> GetSimSpaceBoneNameOptions() const;
+
+	//! Update \c SimSpaceSkeletalMesh and \c SimSpaceTransformIndex according to
+	//! \c RestCollection->TargetSkeletalMesh and SimSpaceBoneName.
+	//! \ret \c true if a valid sim space transform is found.
+	bool UpdateSimSpaceTransformIndex();
+
+	//! Return the rest transform to be used as the simulation space. 
+	//! \c UpdateSimSpaceTransformIndex() must be called prior to calling this function.
+	FTransform GetSimSpaceRestTransform() const;
+
+	//
+	// Render the Procedural Mesh
+	//
+	struct FFleshRenderMesh
+	{
+		TArray<FVector> Vertices;
+		TArray<int32> Triangles;
+		TArray<FVector> Normals;
+		TArray<FVector2D> UVs;
+		TArray<FLinearColor> Colors;
+		TArray<FProcMeshTangent> Tangents;
+	};
+	FFleshRenderMesh* RenderMesh = nullptr;
+	void RenderProceduralMesh();
+	void ResetProceduralMesh();
+
+	bool bBoundsNeedsUpdate = true;
+	FBoxSphereBounds BoundingBox = FBoxSphereBounds(ForceInitToZero);
+
+	FTransform PrevTransform = FTransform::Identity;
+
+	TArray<FVector> GetSkeletalMeshEmbeddedPositionsInternal(const ChaosDeformableBindingOption Format, const FTransform TargetDeformationSkeletonOffset, const FName TargetBone = "", const float SimulationBlendWeight = 1.f, TArray<bool>* OutInfluence = nullptr) const;
+	TArray<FVector> GetEmbeddedPositionsInternal(const TArray<FVector>& InPositions, const FName SkeletalMeshName, const float SimulationBlendWeight = 1.f, TArray<bool>* OutInfluence = nullptr) const;
+	TArray<FVector> GetSkeletalMeshBindingPositionsInternal(const USkeletalMesh* InSkeletalMesh, TArray<bool>* OutInfluence = nullptr) const;
+	void DebugDrawSkeletalMeshBindingPositions() const;
+
+};
+
diff --git a/Engine/Plugins/Experimental/ChaosFlesh/Source/ChaosFleshEngine/Public/ChaosFlesh/FleshActor.h b/Engine/Plugins/Experimental/ChaosFlesh/Source/ChaosFleshEngine/Public/ChaosFlesh/FleshActor.h
--- a/Engine/Plugins/Experimental/ChaosFlesh/Source/ChaosFleshEngine/Public/ChaosFlesh/FleshActor.h
+++ b/Engine/Plugins/Experimental/ChaosFlesh/Source/ChaosFleshEngine/Public/ChaosFlesh/FleshActor.h
@@ -28,7 +28,7 @@ public:
 	TObjectPtr<UFleshComponent> FleshComponent;
 	UFleshComponent* GetFleshComponent() const { return FleshComponent; }
 
-	UPROPERTY(EditAnywhere, BlueprintReadWrite, Category = "Physics")
+	UPROPERTY(EditAnywhere, BlueprintReadWrite, Category = "Physics", meta = (DisplayPriority = 1))
 	TObjectPtr<ADeformableSolverActor> PrimarySolver;
 
 #if WITH_EDITOR
diff --git a/Engine/Plugins/Experimental/ChaosFlesh/Source/ChaosFleshEngine/Public/ChaosFlesh/FleshComponent.h b/Engine/Plugins/Experimental/ChaosFlesh/Source/ChaosFleshEngine/Public/ChaosFlesh/FleshComponent.h
--- a/Engine/Plugins/Experimental/ChaosFlesh/Source/ChaosFleshEngine/Public/ChaosFlesh/FleshComponent.h
+++ b/Engine/Plugins/Experimental/ChaosFlesh/Source/ChaosFleshEngine/Public/ChaosFlesh/FleshComponent.h
@@ -5,6 +5,8 @@
 #include "Chaos/Deformable/ChaosDeformableSolver.h"
 #include "ChaosFlesh/ChaosDeformableSolverThreading.h"
 #include "ChaosFlesh/ChaosDeformablePhysicsComponent.h"
+#include "ChaosFlesh/ChaosDeformableTetrahedralComponent.h"
+#include "ChaosFlesh/ChaosDeformableGameplayComponent.h"
 #include "ChaosFlesh/FleshAsset.h"
 #include "ChaosFlesh/FleshDynamicAsset.h"
 #include "ChaosFlesh/SimulationAsset.h"
@@ -17,161 +19,18 @@ class FFleshCollection;
 class ADeformableSolverActor;
 class UDeformableSolverComponent;
 
-/**
-*  Options for binding positions query.
-*/
-UENUM()
-enum ChaosDeformableBindingOption : uint8
-{
-	WorldPos		UMETA(DisplayName = "World Positions"),
-	WorldDelta		UMETA(DisplayName = "World Deltas"),
-	ComponentPos    UMETA(DisplayName = "Component Positions"),
-	ComponentDelta  UMETA(DisplayName = "Component Deltas"),
-	BonePos			UMETA(DisplayName = "Bone Positions"),
-	BoneDelta		UMETA(DisplayName = "Bone Deltas"),
-};
+
 
 
 /**
 *	FleshComponent
 */
 UCLASS(meta = (BlueprintSpawnableComponent))
-class CHAOSFLESHENGINE_API UFleshComponent : public UDeformablePhysicsComponent
+class CHAOSFLESHENGINE_API UFleshComponent : public UDeformableGameplayComponent
 {
 	GENERATED_UCLASS_BODY()
-
 public:
-	typedef Chaos::Softs::FFleshThreadingProxy FFleshThreadingProxy;
-
-	~UFleshComponent();
-
-	/** USceneComponent Interface */
-	virtual void BeginPlay() override;
-	virtual void OnRegister() override;
-	virtual void TickComponent(float DeltaTime, enum ELevelTick TickType, FActorComponentTickFunction *ThisTickFunction) override;
-	virtual FBoxSphereBounds CalcBounds(const FTransform& LocalToWorld) const override;
-	void UpdateLocalBounds();
-	virtual void EndPlay(const EEndPlayReason::Type ReasonEnd) override;
-	void Invalidate();
-
-	/** Simulation Interface*/
-	virtual FThreadingProxy* NewProxy() override;
-	virtual FDataMapValue NewDeformableData() override;
-	virtual void UpdateFromSimualtion(const FDataMapValue* SimualtionBuffer) override;
-
-
-	/** RestCollection */
-	void SetRestCollection(const UFleshAsset * InRestCollection);
-	const UFleshAsset* GetRestCollection() const { return RestCollection; }
-
-
-	/** DynamicCollection */
-	void ResetDynamicCollection();
-	UFleshDynamicAsset* GetDynamicCollection() { return DynamicCollection; }
-	const UFleshDynamicAsset* GetDynamicCollection() const { return DynamicCollection; }
-
-	/** SimulationCollection */
-	void ResetSimulationCollection();
-	USimulationAsset* GetSimulationCollection() { return SimulationCollection; }
-	const USimulationAsset* GetSimulationCollection() const { return SimulationCollection; }
-
-	UPROPERTY(EditAnywhere, Category = "Rendering")
-	TObjectPtr<UProceduralMeshComponent> Mesh;
-
-	/** @deprecated Use GetSkeletalMeshEmbeddedPositions() instead. */
-	UFUNCTION(BlueprintCallable, Category = "Physics", meta = (DeprecatedFunction, DeprecationMessage = "Use GetSkeletalMeshEmbeddedPositions() instead."))
-	TArray<FVector> GetSkeletalMeshBindingPositions(const USkeletalMesh* InSkeletalMesh) const;
-
-	/**
-	* Get the current positions of the transformation hierarchy from \c TargetDeformationSkeleton,
-	* deformed by the tetrahedral mesh.  Results can be in world space postions/deltas, component space
-	* positions/deltas, or bone space positions/deltas.  If a bone space is desired \p TargetBone
-	* must indicate which bone to use. TargetDeformationSkeletonOffset is an offset transform that moves 
-	* the \c TargetDeformationSkeleton to be co-located with the flesh mesh.
-	*/
-	UFUNCTION(BlueprintCallable, Category = "Physics")
-	TArray<FVector> GetSkeletalMeshEmbeddedPositions(const ChaosDeformableBindingOption Format, const FTransform TargetDeformationSkeletonOffset, const FName TargetBone = "", const float SimulationBlendWeight = 1.f) const;
-
-	/** Indices of tetrahedra to hide. */
-	UPROPERTY(EditAnywhere, Category = "Rendering")
-	TSet<int32> HideTetrahedra;
-
-private:
-	/** FleshAsset that describes the simulation rest state. */
-	UPROPERTY(EditAnywhere, Category = "Physics")
-	TObjectPtr<const UFleshAsset> RestCollection;
-
-	/** Current simulation state. */
-	UPROPERTY()
-	TObjectPtr<UFleshDynamicAsset> DynamicCollection;
-
-	/** Simulator input */
-	UPROPERTY()
-	TObjectPtr<USimulationAsset> SimulationCollection;
-
-	//
-	// Sim Space
-	//
-
-	/** Space the simulation will run in. */
-	UPROPERTY(EditAnywhere, Category = "Physics")
-	TEnumAsByte<ChaosDeformableSimSpace> SimSpace = ChaosDeformableSimSpace::World;
-
-	/** 
-	* Bone from the associated skeletal mesh (indicated by RestCollection.TargetSkeletalMesh) to use as 
-	* the space the sim runs in.
-	*/
-	UPROPERTY(EditAnywhere, Category = "Physics", meta = (GetOptions = "GetSimSpaceBoneNameOptions", EditCondition = "SimSpace == ChaosDeformableSimSpace::Bone"))
-	FName SimSpaceBoneName;
-
-	/** The skeletal mesh to use pull the \c SimSpaceBoneName from. */
-	UPROPERTY()
-	TObjectPtr<USkeletalMesh> SimSpaceSkeletalMesh;
-
-	/* Returns a list of bone names from the currently selected skeletal mesh. */
-	UFUNCTION(CallInEditor)
-	TArray<FString> GetSimSpaceBoneNameOptions() const;
-
-	//! Update \c SimSpaceSkeletalMesh and \c SimSpaceTransformIndex according to
-	//! \c RestCollection->TargetSkeletalMesh and SimSpaceBoneName.
-	//! \ret \c true if a valid sim space transform is found.
-	bool UpdateSimSpaceTransformIndex();
-
-	//! Return the rest transform to be used as the simulation space. 
-	//! \c UpdateSimSpaceTransformIndex() must be called prior to calling this function.
-	FTransform GetSimSpaceRestTransform() const;
-
-	//
-	// Render the Procedural Mesh
-	//
-	struct FFleshRenderMesh
-	{
-		TArray<FVector> Vertices;
-		TArray<int32> Triangles;
-		TArray<FVector> Normals;
-		TArray<FVector2D> UVs;
-		TArray<FLinearColor> Colors;
-		TArray<FProcMeshTangent> Tangents;
-	};
-	FFleshRenderMesh* RenderMesh = nullptr;
-	void RenderProceduralMesh();
-	void ResetProceduralMesh();
-
-	bool bBoundsNeedsUpdate = true;
-	FBoxSphereBounds BoundingBox = FBoxSphereBounds(ForceInitToZero);
-
-	FTransform PrevTransform = FTransform::Identity;
-
-	TArray<FVector> GetSkeletalMeshEmbeddedPositionsInternal(const ChaosDeformableBindingOption Format, const FTransform TargetDeformationSkeletonOffset, const FName TargetBone = "", const float SimulationBlendWeight = 1.f, TArray<bool>* OutInfluence = nullptr) const;
-	TArray<FVector> GetEmbeddedPositionsInternal(const TArray<FVector>& InPositions, const FName SkeletalMeshName, const float SimulationBlendWeight = 1.f, TArray<bool>* OutInfluence = nullptr) const;
-	TArray<FVector> GetSkeletalMeshBindingPositionsInternal(const USkeletalMesh* InSkeletalMesh, TArray<bool>* OutInfluence = nullptr) const;
-	void DebugDrawSkeletalMeshBindingPositions() const;
-
-	int32 SimSpaceTransformIndex = INDEX_NONE;
-	int32 SimSpaceTransformGlobalIndex = INDEX_NONE;
 
-	TSet<int32> FaceVerticesToSkip;
-	int32 NumSkippedFaces = 0;
-	int32 NumSkippedTets = 0;
-};
+	TArray<int32> HideTetrahedra;
 
+};
\ No newline at end of file
diff --git a/Engine/Source/Runtime/Experimental/Chaos/Private/Chaos/Deformable/ChaosDeformableSolver.cpp b/Engine/Source/Runtime/Experimental/Chaos/Private/Chaos/Deformable/ChaosDeformableSolver.cpp
--- a/Engine/Source/Runtime/Experimental/Chaos/Private/Chaos/Deformable/ChaosDeformableSolver.cpp
+++ b/Engine/Source/Runtime/Experimental/Chaos/Private/Chaos/Deformable/ChaosDeformableSolver.cpp
@@ -1430,7 +1430,7 @@ namespace Chaos::Softs
 
 				WriteTrisGEO(Evolution->Particles(), *SurfaceElements);
 				FString file = FPaths::ProjectDir();
-				file.Append(TEXT("/HoudiniOutput/DtLog.txt"));
+				file.Append(TEXT("/DebugOutput/DtLog.txt"));
 				if (Frame == 0)
 				{
 					FFileHelper::SaveStringToFile(FString(TEXT("DeltaTime\r\n")), *file);
@@ -1443,8 +1443,7 @@ namespace Chaos::Softs
 	void FDeformableSolver::WriteTrisGEO(const FSolverParticles& Particles, const TArray<TVec3<int32>>& Mesh)
 	{
 		FString file = FPaths::ProjectDir();
-		//file.Append(TEXT("\HoudiniOuput\Test.geo"));
-		file.Append(TEXT("/HoudiniOutput/sim_frame_"));
+		file.Append(TEXT("/DebugOutput/sim_frame_"));
 		file.Append(FString::FromInt(Frame));
 		file.Append(TEXT(".geo"));
 
