Adjust shader path for cook compatibility.

diff --git a/Engine/Plugins/Experimental/ChaosFlesh/Shaders/Private/DIFleshDeformer.ush b/Engine/Plugins/Experimental/ChaosFlesh/Shaders/Private/DIFleshDeformer.ush
new file mode 100644
--- /dev/null
+++ b/Engine/Plugins/Experimental/ChaosFlesh/Shaders/Private/DIFleshDeformer.ush
@@ -0,0 +1,221 @@
+// Copyright Epic Games, Inc.  All Rights Reserved.
+
+uint {DataInterfaceName}_NumVertices;
+uint {DataInterfaceName}_BaseVertexIndex; // InputStreamStart
+
+uint {DataInterfaceName}_NumTetVertexBuffer;
+uint {DataInterfaceName}_NumTetRestVertexBuffer;
+
+uint {DataInterfaceName}_NumParentsBuffer;
+uint {DataInterfaceName}_NumWeightsBuffer;
+uint {DataInterfaceName}_NumOffsetBuffer;
+uint {DataInterfaceName}_NumMaskBuffer;
+
+Buffer<float> {DataInterfaceName}_TetRestVertexBuffer;
+Buffer<float> {DataInterfaceName}_TetVertexBuffer;
+
+Buffer<int>	{DataInterfaceName}_ParentsBuffer;
+Buffer<float> {DataInterfaceName}_WeightsBuffer;
+Buffer<float> {DataInterfaceName}_OffsetBuffer;
+Buffer<float> {DataInterfaceName}_MaskBuffer;
+
+uint ReadNumVertices_{DataInterfaceName}(uint Index)
+{
+	return {DataInterfaceName}_NumVertices;
+}
+
+//
+// Tet mesh points, rest and dynamic 
+//
+
+float3 ReadRestTetVertex_{DataInterfaceName}(uint Index)
+{
+#if ENABLE_DEFORMER_FLESH
+	if(Index >= 0 && Index < {DataInterfaceName}_NumTetRestVertexBuffer)
+	{
+		return float3(
+			{DataInterfaceName}_TetRestVertexBuffer[Index * 3], 
+			{DataInterfaceName}_TetRestVertexBuffer[Index * 3 + 1], 
+			{DataInterfaceName}_TetRestVertexBuffer[Index * 3 + 2]);
+	}
+#endif
+	return float3(0,0,0);
+}
+
+float3 ReadTetVertex_{DataInterfaceName}(uint Index)
+{
+#if ENABLE_DEFORMER_FLESH
+	if(Index >= 0 && Index < {DataInterfaceName}_NumTetVertexBuffer)
+	{
+		return float3(
+			{DataInterfaceName}_TetVertexBuffer[Index * 3], 
+			{DataInterfaceName}_TetVertexBuffer[Index * 3 + 1], 
+			{DataInterfaceName}_TetVertexBuffer[Index * 3 + 2]);
+	}
+#endif
+	return float3(0,0,0);
+}
+
+//
+// Render mesh to tet bindings
+//
+
+int4 ReadParents_{DataInterfaceName}(uint VertexIndex)
+{
+#if ENABLE_DEFORMER_FLESH
+	if(VertexIndex < {DataInterfaceName}_NumParentsBuffer)
+	{
+		uint BufferIndex = VertexIndex + {DataInterfaceName}_BaseVertexIndex;
+		return int4(
+			{DataInterfaceName}_ParentsBuffer[BufferIndex * 4], 
+			{DataInterfaceName}_ParentsBuffer[BufferIndex * 4 + 1], 
+			{DataInterfaceName}_ParentsBuffer[BufferIndex * 4 + 2],
+			{DataInterfaceName}_ParentsBuffer[BufferIndex * 4 + 3]);
+	}
+#endif
+	return int4(-1,-1,-1,-1);
+}
+
+float4 ReadWeights_{DataInterfaceName}(uint VertexIndex)
+{
+#if ENABLE_DEFORMER_FLESH
+	if(VertexIndex < {DataInterfaceName}_NumWeightsBuffer)
+	{
+		uint BufferIndex = VertexIndex + {DataInterfaceName}_BaseVertexIndex;
+		return float4(
+			{DataInterfaceName}_WeightsBuffer[BufferIndex * 4], 
+			{DataInterfaceName}_WeightsBuffer[BufferIndex * 4 + 1], 
+			{DataInterfaceName}_WeightsBuffer[BufferIndex * 4 + 2],
+			{DataInterfaceName}_WeightsBuffer[BufferIndex * 4 + 3]);
+	}
+#endif
+	return float4(0,0,0,0);
+}
+
+float3 ReadOffset_{DataInterfaceName}(uint VertexIndex)
+{
+#if ENABLE_DEFORMER_FLESH
+	if(VertexIndex < {DataInterfaceName}_NumOffsetBuffer)
+	{
+		uint BufferIndex = VertexIndex + {DataInterfaceName}_BaseVertexIndex;
+		return float3(
+			{DataInterfaceName}_OffsetBuffer[BufferIndex * 3], 
+			{DataInterfaceName}_OffsetBuffer[BufferIndex * 3 + 1], 
+			{DataInterfaceName}_OffsetBuffer[BufferIndex * 3 + 2]);
+	}
+#endif
+	return float3(0,0,0);
+}
+
+float ReadMask_{DataInterfaceName}(uint VertexIndex)
+{
+#if ENABLE_DEFORMER_FLESH
+	if(VertexIndex < {DataInterfaceName}_NumMaskBuffer)
+	{
+		uint BufferIndex = VertexIndex + {DataInterfaceName}_BaseVertexIndex;
+		return {DataInterfaceName}_MaskBuffer[BufferIndex];
+	}
+#endif
+	return 0.0;
+}
+
+//
+// Triangle orientation
+//
+
+
+// [ Duff et al. 2017, "Building an Orthonormal Basis, Revisited" ]
+// Discontinuity at TangentZ.z == 0
+float3x3 GetTangentBasis( float3 TangentZ )
+{
+    const float Sign = TangentZ.z >= 0 ? 1 : -1;
+    const float a = -rcp( Sign + TangentZ.z );
+    const float b = TangentZ.x * TangentZ.y * a;
+
+    float3 TangentX = { 1 + Sign * a * Pow2( TangentZ.x ), Sign * b, -Sign * TangentZ.x };
+    float3 TangentY = { b,  Sign + a * Pow2( TangentZ.y ), -TangentZ.y };
+
+    return float3x3( TangentX, TangentY, TangentZ );
+}
+
+float3 SafeNormalize(float3 Vec)
+{
+	float L = length(Vec);
+	if(L > 1.0e-6)
+	{
+		return Vec / L;
+	}
+	return float3(1,0,0);
+}
+
+float3x3 GetOrthogonalBasisVectors(float3 PtA, float3 PtB, float3 PtC)
+{
+	float3 EdgeBA = PtB - PtA;
+	float3 EdgeCA = PtC - PtA;
+	float3 OrthoNorm = SafeNormalize(cross(EdgeBA, EdgeCA));
+	return GetTangentBasis(OrthoNorm);
+/*
+	float3 EdgeBA = SafeNormalize(PtB - PtA);
+	float3 EdgeCA = SafeNormalize(PtC - PtA);
+	float3 OrthoNorm = cross(EdgeBA, EdgeCA);
+	float3 OrthoCA = cross(EdgeBA, OrthoNorm);
+	return float3x3(
+		EdgeBA[0], OrthoCA[0], OrthoNorm[0],
+		EdgeBA[1], OrthoCA[1], OrthoNorm[1],
+		EdgeBA[2], OrthoCA[2], OrthoNorm[2]); // Col major
+*/
+}
+
+
+float3x3 GetRestOrthogonalBasisVectors(int4 Tet) // assumes Tet[3] == -1
+{
+	float3 PtA = ReadRestTetVertex_{DataInterfaceName}(Tet[0]);
+	float3 PtB = ReadRestTetVertex_{DataInterfaceName}(Tet[1]);
+	float3 PtC = ReadRestTetVertex_{DataInterfaceName}(Tet[2]);
+	return GetOrthogonalBasisVectors(PtA, PtB, PtC);
+}
+
+float3x3 GetDynamicOrthogonalBasisVectors(int4 Tet) // assumes Tet[3] == -1
+{
+	float3 PtA = ReadTetVertex_{DataInterfaceName}(Tet[0]);
+	float3 PtB = ReadTetVertex_{DataInterfaceName}(Tet[1]);
+	float3 PtC = ReadTetVertex_{DataInterfaceName}(Tet[2]);
+	return GetOrthogonalBasisVectors(PtA, PtB, PtC);
+}
+
+float3 GetRotatedOffsetVector(uint VertexIndex, int4 Tet)
+{
+	float3 Offset = { 0, 0, 0 };
+#if ENABLE_DEFORMER_FLESH
+	Offset = ReadOffset_{DataInterfaceName}(VertexIndex);
+	if(Offset[0] != 0 && Offset[1] != 0 && Offset[2] != 0)
+	{
+		float3x3 RestRotInv = transpose(GetRestOrthogonalBasisVectors(Tet)); // Ortho matrix, so T == I
+		float3x3 CurrRot = GetDynamicOrthogonalBasisVectors(Tet);
+		Offset = mul(Offset, mul(RestRotInv, CurrRot));
+	}
+#endif
+	return Offset;
+}
+
+//
+// Helpers
+//
+
+float3 GetEmbeddedPos_{DataInterfaceName}(uint VertexIndex)
+{
+	float3 Pos = { 0, 0, 0 };
+#if ENABLE_DEFORMER_FLESH
+	int4 Parents = ReadParents_{DataInterfaceName}(VertexIndex);
+	float4 Weights = ReadWeights_{DataInterfaceName}(VertexIndex);
+	UNROLL for(uint i=0; i < 4; i++)
+	{
+		if(Parents[i] >= 0 && Weights[i] != 0.0)
+		{
+			Pos += ReadTetVertex_{DataInterfaceName}(Parents[i]) * Weights[i];
+		}
+	}
+	Pos += GetRotatedOffsetVector(VertexIndex, Parents);
+#endif
+	return Pos;
+}
diff --git a/Engine/Plugins/Experimental/ChaosFlesh/Source/ChaosFleshEngine/Private/ChaosFlesh/ChaosFleshEnginePlugin.cpp b/Engine/Plugins/Experimental/ChaosFlesh/Source/ChaosFleshEngine/Private/ChaosFlesh/ChaosFleshEnginePlugin.cpp
--- a/Engine/Plugins/Experimental/ChaosFlesh/Source/ChaosFleshEngine/Private/ChaosFlesh/ChaosFleshEnginePlugin.cpp
+++ b/Engine/Plugins/Experimental/ChaosFlesh/Source/ChaosFleshEngine/Private/ChaosFlesh/ChaosFleshEnginePlugin.cpp
@@ -26,8 +26,8 @@ void FChaosFleshEnginePlugin::StartupModule()
 	const FString PluginShaderDir = 
 		FPaths::Combine(
 			IPluginManager::Get().FindPlugin(TEXT("ChaosFlesh"))->GetBaseDir(), 
-			TEXT("Source/ChaosFleshEngine/Shaders/Private"));
-	AddShaderSourceDirectoryMapping(TEXT("/Plugin/ChaosFleshEngine"), PluginShaderDir);
+			TEXT("Shaders"));
+	AddShaderSourceDirectoryMapping(TEXT("/Plugin/ChaosFlesh"), PluginShaderDir);
 
 	// Self-register for now.
 	// FIXME: Should add a method to IOptimusModule to register interfaces.
diff --git a/Engine/Plugins/Experimental/ChaosFlesh/Source/ChaosFleshEngine/Shaders/Private/DIFleshDeformer.ush b/Engine/Plugins/Experimental/ChaosFlesh/Source/ChaosFleshEngine/Shaders/Private/DIFleshDeformer.ush
deleted file mode 100644
