Correct shader path.

diff --git a/Engine/Plugins/Experimental/ChaosFlesh/Source/ChaosFleshEngine/Private/DataInterfaces/DIFleshDeformer.cpp b/Engine/Plugins/Experimental/ChaosFlesh/Source/ChaosFleshEngine/Private/DataInterfaces/DIFleshDeformer.cpp
--- a/Engine/Plugins/Experimental/ChaosFlesh/Source/ChaosFleshEngine/Private/DataInterfaces/DIFleshDeformer.cpp
+++ b/Engine/Plugins/Experimental/ChaosFlesh/Source/ChaosFleshEngine/Private/DataInterfaces/DIFleshDeformer.cpp
@@ -117,10 +117,9 @@ UDIFleshDeformer::GetPermutations(FComputeKernelPermutationVector& OutPermutatio
 void
 UDIFleshDeformer::GetShaderHash(FString& InOutKey) const
 {
-	// Module registers "ChaosFlesh/Source/ChaosFleshEngine/Shaders/Private" as virtual shader 
-	// path: "/Plugin/ChaosFleshEngine".
+	// Module registers "ChaosFlesh/Shaders/Private" as virtual shader path: "/Plugin/ChaosFlesh".
 	GetShaderFileHash(
-		TEXT("/Plugin/ChaosFleshEngine/DIFleshDeformer.ush"),
+		TEXT("/Plugin/ChaosFlesh/Private/DIFleshDeformer.ush"),
 		EShaderPlatform::SP_PCD3D_SM5).AppendString(InOutKey);
 }
 
@@ -134,7 +133,7 @@ UDIFleshDeformer::GetHLSL(FString& OutHLSL, FString const& InDataInterfaceName)
 
 	FString TemplateFile;
 	LoadShaderSourceFile(
-		TEXT("/Plugin/ChaosFleshEngine/DIFleshDeformer.ush"),
+		TEXT("/Plugin/ChaosFlesh/Private/DIFleshDeformer.ush"),
 		EShaderPlatform::SP_PCD3D_SM5, &TemplateFile, nullptr);
 	OutHLSL += FString::Format(*TemplateFile, TemplateArgs);
 }
