FleshDeformer updates. Remove "TargetSkeletalMesh" name identifier in favor of disambigution via "TargetDeformationSkeleton".  Transfer ownership of GPU buffers to the tetrahedral component.  Upload unvarynig data once, varying data only on tick when it's changed.  Sparsify bindings and vertex data to the minimal data set required, and reindex as necessary.  Use 16 bit floats for most floating point types.  Drop down to 16 or 8 bit uint for integral data when possible.

Adds 3 boolean pvars for testing in editor:

p.Chaos.Deformable.FleshDeformer.TestSparseMappings
p.Chaos.Deformable.FleshDeformer.TestCompressedIndexing
p.Chaos.Deformable.FleshDeformer.TestUnpacking

diff --git a/Engine/Plugins/Experimental/ChaosFlesh/Shaders/Private/DIFleshDeformer.ush b/Engine/Plugins/Experimental/ChaosFlesh/Shaders/Private/DIFleshDeformer.ush
--- a/Engine/Plugins/Experimental/ChaosFlesh/Shaders/Private/DIFleshDeformer.ush
+++ b/Engine/Plugins/Experimental/ChaosFlesh/Shaders/Private/DIFleshDeformer.ush
@@ -6,7 +6,14 @@ uint {DataInterfaceName}_BaseVertexIndex; // InputStreamStart
 uint {DataInterfaceName}_NumTetVertexBuffer;
 uint {DataInterfaceName}_NumTetRestVertexBuffer;
 
+uint {DataInterfaceName}_NumVertsToParentsBuffer;
+int {DataInterfaceName}_VertsToParentsBufferOffset;
+uint {DataInterfaceName}_VertsToParentsBufferStride;
+
 uint {DataInterfaceName}_NumParentsBuffer;
+int {DataInterfaceName}_ParentsBufferOffset;
+uint {DataInterfaceName}_ParentsBufferStride;	// 1 = 8 bit, 2 = 16 bit, 4 = 32 bit
+
 uint {DataInterfaceName}_NumWeightsBuffer;
 uint {DataInterfaceName}_NumOffsetBuffer;
 uint {DataInterfaceName}_NumMaskBuffer;
@@ -14,10 +21,11 @@ uint {DataInterfaceName}_NumMaskBuffer;
 Buffer<float> {DataInterfaceName}_TetRestVertexBuffer;
 Buffer<float> {DataInterfaceName}_TetVertexBuffer;
 
-Buffer<int>	{DataInterfaceName}_ParentsBuffer;
-Buffer<float> {DataInterfaceName}_WeightsBuffer;
-Buffer<float> {DataInterfaceName}_OffsetBuffer;
-Buffer<float> {DataInterfaceName}_MaskBuffer;
+Buffer<uint> {DataInterfaceName}_VertsToParentsBuffer;
+Buffer<uint> {DataInterfaceName}_ParentsBuffer;
+Buffer<half> {DataInterfaceName}_WeightsBuffer;
+Buffer<half> {DataInterfaceName}_OffsetBuffer;
+Buffer<half> {DataInterfaceName}_MaskBuffer;
 
 uint ReadNumVertices_{DataInterfaceName}(uint Index)
 {
@@ -28,10 +36,10 @@ uint ReadNumVertices_{DataInterfaceName}(uint Index)
 // Tet mesh points, rest and dynamic 
 //
 
-float3 ReadRestTetVertex_{DataInterfaceName}(uint Index)
+float3 ReadRestTetVertex_{DataInterfaceName}(int Index)
 {
 #if ENABLE_DEFORMER_FLESH
-	if(Index >= 0 && Index < {DataInterfaceName}_NumTetRestVertexBuffer)
+	if(Index >= 0 && Index*3+2 < {DataInterfaceName}_NumTetRestVertexBuffer)
 	{
 		return float3(
 			{DataInterfaceName}_TetRestVertexBuffer[Index * 3], 
@@ -42,10 +50,10 @@ float3 ReadRestTetVertex_{DataInterfaceName}(uint Index)
 	return float3(0,0,0);
 }
 
-float3 ReadTetVertex_{DataInterfaceName}(uint Index)
+float3 ReadTetVertex_{DataInterfaceName}(int Index)
 {
 #if ENABLE_DEFORMER_FLESH
-	if(Index >= 0 && Index < {DataInterfaceName}_NumTetVertexBuffer)
+	if(Index >= 0 && Index*3+2 < {DataInterfaceName}_NumTetVertexBuffer)
 	{
 		return float3(
 			{DataInterfaceName}_TetVertexBuffer[Index * 3], 
@@ -60,54 +68,146 @@ float3 ReadTetVertex_{DataInterfaceName}(uint Index)
 // Render mesh to tet bindings
 //
 
+// Stride 1 = 8 bit, 2 = 16 bit, 4 = 32 bit - based off of sizeof(<index type>)
+uint UnpackUInt(const Buffer<uint> Array, const uint Index, const uint Stride, const uint NumValues)
+{
+	// NumValues reflects the number of packed values in Array
+	if(Index >= NumValues)
+		return 0;
+	
+	uint ValsPerEntry = 4 / Stride;
+	uint MajorIndex = floor(Index / ValsPerEntry);
+
+	uint Val = Array[MajorIndex];
+
+	uint MinorIndex = Index % ValsPerEntry;
+	return BitFieldExtractU32(Val, Stride*8, Stride*8*MinorIndex);
+}
+
+int UnpackInt_32bit(const Buffer<uint> Array, const uint Index, const uint ArraySize)
+{
+	return Index >= ArraySize ? -1 : int(Array[Index]);
+}
+
+int UnpackInt_16bit(const Buffer<uint> Array, const uint Index, const uint ArraySize)
+{
+	// Array: [..., <A,B>, ...], where A-B are 16 bit ints, and Index refers to one of A or B.
+	uint BaseIndex = floor(Index / 2); // 32 bit index
+	uint SubIndex = Index % 2;  // [0, 1], corrsponding to A or B
+	return BaseIndex >= ArraySize ? -1 : int((Array[BaseIndex] >> (SubIndex * 16)) & 0xFFFF);
+}
+
+int UnpackInt_8bit(const Buffer<uint> Array, const uint Index, const uint ArraySize)
+{
+	// Array: [..., <A,B,C,D>, ...], where A-D are 8 bit ints, and Index refers to one of A-D.
+	uint BaseIndex = floor(Index / 4); // 32 bit index
+	uint SubIndex = Index % 4;  // [0, 3], corresponding to A, B, C, or D
+	return BaseIndex >= ArraySize ? -1 : int((Array[BaseIndex] >> (SubIndex * 8)) & 0xFF);
+}
+
+int4 ReadParents32_{DataInterfaceName}(uint BufferIndex)
+{
+#if ENABLE_DEFORMER_FLESH
+	return int4(
+		UnpackInt_32bit({DataInterfaceName}_ParentsBuffer, BufferIndex    , {DataInterfaceName}_NumParentsBuffer) - {DataInterfaceName}_ParentsBufferOffset, 
+		UnpackInt_32bit({DataInterfaceName}_ParentsBuffer, BufferIndex + 1, {DataInterfaceName}_NumParentsBuffer) - {DataInterfaceName}_ParentsBufferOffset, 
+		UnpackInt_32bit({DataInterfaceName}_ParentsBuffer, BufferIndex + 2, {DataInterfaceName}_NumParentsBuffer) - {DataInterfaceName}_ParentsBufferOffset, 
+		UnpackInt_32bit({DataInterfaceName}_ParentsBuffer, BufferIndex + 3, {DataInterfaceName}_NumParentsBuffer) - {DataInterfaceName}_ParentsBufferOffset);
+#endif
+	return int4(-1,-1,-1,-1);
+}
+
+int4 ReadParents16_{DataInterfaceName}(uint BufferIndex)
+{
+#if ENABLE_DEFORMER_FLESH
+	return int4(
+		UnpackInt_16bit({DataInterfaceName}_ParentsBuffer, BufferIndex    , {DataInterfaceName}_NumParentsBuffer) - {DataInterfaceName}_ParentsBufferOffset, 
+		UnpackInt_16bit({DataInterfaceName}_ParentsBuffer, BufferIndex + 1, {DataInterfaceName}_NumParentsBuffer) - {DataInterfaceName}_ParentsBufferOffset, 
+		UnpackInt_16bit({DataInterfaceName}_ParentsBuffer, BufferIndex + 2, {DataInterfaceName}_NumParentsBuffer) - {DataInterfaceName}_ParentsBufferOffset, 
+		UnpackInt_16bit({DataInterfaceName}_ParentsBuffer, BufferIndex + 3, {DataInterfaceName}_NumParentsBuffer) - {DataInterfaceName}_ParentsBufferOffset);
+#endif
+	return int4(-1,-1,-1,-1);
+}
+
+int4 ReadParents8_{DataInterfaceName}(uint BufferIndex)
+{
+#if ENABLE_DEFORMER_FLESH
+	return int4(
+		UnpackInt_8bit({DataInterfaceName}_ParentsBuffer, BufferIndex    , {DataInterfaceName}_NumParentsBuffer) - {DataInterfaceName}_ParentsBufferOffset, 
+		UnpackInt_8bit({DataInterfaceName}_ParentsBuffer, BufferIndex + 1, {DataInterfaceName}_NumParentsBuffer) - {DataInterfaceName}_ParentsBufferOffset, 
+		UnpackInt_8bit({DataInterfaceName}_ParentsBuffer, BufferIndex + 2, {DataInterfaceName}_NumParentsBuffer) - {DataInterfaceName}_ParentsBufferOffset, 
+		UnpackInt_8bit({DataInterfaceName}_ParentsBuffer, BufferIndex + 3, {DataInterfaceName}_NumParentsBuffer) - {DataInterfaceName}_ParentsBufferOffset);
+#endif
+	return int4(-1,-1,-1,-1);
+}
+
 int4 ReadParents_{DataInterfaceName}(uint VertexIndex)
 {
 #if ENABLE_DEFORMER_FLESH
-	if(VertexIndex < {DataInterfaceName}_NumParentsBuffer)
+	// The shader only sees one section of the mesh at a time, each of which is indexed starting at 0,
+	// so we need to offset VertexIndex by the base vertex index, to get our buffer relative index.
+	uint BufferIndex = VertexIndex + {DataInterfaceName}_BaseVertexIndex;
+
+	// Many render vertices will likely share the same set of parents, so the set of parents is sparsely
+	// represented.  The VertsToParentsBuffer contains the mapping from render vertices to parents.
+	int ParentsIndex = 
+		UnpackUInt(
+			{DataInterfaceName}_VertsToParentsBuffer, 
+			BufferIndex, 
+			{DataInterfaceName}_VertsToParentsBufferStride, 
+			{DataInterfaceName}_NumVertsToParentsBuffer);
+	ParentsIndex -= {DataInterfaceName}_VertsToParentsBufferOffset;
+	if(ParentsIndex >= 0)
 	{
-		uint BufferIndex = VertexIndex + {DataInterfaceName}_BaseVertexIndex;
-		return int4(
-			{DataInterfaceName}_ParentsBuffer[BufferIndex * 4], 
-			{DataInterfaceName}_ParentsBuffer[BufferIndex * 4 + 1], 
-			{DataInterfaceName}_ParentsBuffer[BufferIndex * 4 + 2],
-			{DataInterfaceName}_ParentsBuffer[BufferIndex * 4 + 3]);
+		// That gives us the index of the int4, which we have laid out in a flat array, so mult by 4.
+		// The parents buffer is potentially packed 8, 16, or 32 bit numbers, but we'll deal with
+		// that later.
+		ParentsIndex *= 4;
+
+		// Do lookups into the parents buffer, unpacking as necessary depending on the storage format.	
+		if({DataInterfaceName}_ParentsBufferStride == 4)
+			return ReadParents32_{DataInterfaceName}(ParentsIndex);
+		else if({DataInterfaceName}_ParentsBufferStride == 2)
+			return ReadParents16_{DataInterfaceName}(ParentsIndex);
+		else // if 1
+			return ReadParents8_{DataInterfaceName}(ParentsIndex);
 	}
 #endif
 	return int4(-1,-1,-1,-1);
 }
 
-float4 ReadWeights_{DataInterfaceName}(uint VertexIndex)
+half4 ReadWeights_{DataInterfaceName}(uint VertexIndex)
 {
 #if ENABLE_DEFORMER_FLESH
 	if(VertexIndex < {DataInterfaceName}_NumWeightsBuffer)
 	{
 		uint BufferIndex = VertexIndex + {DataInterfaceName}_BaseVertexIndex;
-		return float4(
+		return half4(
 			{DataInterfaceName}_WeightsBuffer[BufferIndex * 4], 
 			{DataInterfaceName}_WeightsBuffer[BufferIndex * 4 + 1], 
 			{DataInterfaceName}_WeightsBuffer[BufferIndex * 4 + 2],
 			{DataInterfaceName}_WeightsBuffer[BufferIndex * 4 + 3]);
 	}
 #endif
-	return float4(0,0,0,0);
+	return half4(0,0,0,0);
 }
 
-float3 ReadOffset_{DataInterfaceName}(uint VertexIndex)
+half3 ReadOffset_{DataInterfaceName}(uint VertexIndex)
 {
 #if ENABLE_DEFORMER_FLESH
 	if(VertexIndex < {DataInterfaceName}_NumOffsetBuffer)
 	{
 		uint BufferIndex = VertexIndex + {DataInterfaceName}_BaseVertexIndex;
-		return float3(
+		return half3(
 			{DataInterfaceName}_OffsetBuffer[BufferIndex * 3], 
 			{DataInterfaceName}_OffsetBuffer[BufferIndex * 3 + 1], 
 			{DataInterfaceName}_OffsetBuffer[BufferIndex * 3 + 2]);
 	}
 #endif
-	return float3(0,0,0);
+	return half3(0,0,0);
 }
 
-float ReadMask_{DataInterfaceName}(uint VertexIndex)
+half ReadMask_{DataInterfaceName}(uint VertexIndex)
 {
 #if ENABLE_DEFORMER_FLESH
 	if(VertexIndex < {DataInterfaceName}_NumMaskBuffer)
@@ -116,14 +216,13 @@ float ReadMask_{DataInterfaceName}(uint VertexIndex)
 		return {DataInterfaceName}_MaskBuffer[BufferIndex];
 	}
 #endif
-	return 0.0;
+	return half(0);
 }
 
 //
 // Triangle orientation
 //
 
-
 // [ Duff et al. 2017, "Building an Orthonormal Basis, Revisited" ]
 // Discontinuity at TangentZ.z == 0
 float3x3 GetTangentBasis( float3 TangentZ )
@@ -207,13 +306,11 @@ float3 GetEmbeddedPos_{DataInterfaceName}(uint VertexIndex)
 	float3 Pos = { 0, 0, 0 };
 #if ENABLE_DEFORMER_FLESH
 	int4 Parents = ReadParents_{DataInterfaceName}(VertexIndex);
-	float4 Weights = ReadWeights_{DataInterfaceName}(VertexIndex);
+	half4 Weights = ReadWeights_{DataInterfaceName}(VertexIndex);
 	UNROLL for(uint i=0; i < 4; i++)
 	{
-		if(Parents[i] >= 0 && Weights[i] != 0.0)
-		{
-			Pos += ReadTetVertex_{DataInterfaceName}(Parents[i]) * Weights[i];
-		}
+		// Returns zero vec if out of range.
+		Pos += ReadTetVertex_{DataInterfaceName}(Parents[i]) * Weights[i];
 	}
 	Pos += GetRotatedOffsetVector(VertexIndex, Parents);
 #endif
diff --git a/Engine/Plugins/Experimental/ChaosFlesh/Source/ChaosFleshEditor/Private/ChaosFlesh/Asset/ActorFactoryFlesh.cpp b/Engine/Plugins/Experimental/ChaosFlesh/Source/ChaosFleshEditor/Private/ChaosFlesh/Asset/ActorFactoryFlesh.cpp
--- a/Engine/Plugins/Experimental/ChaosFlesh/Source/ChaosFleshEditor/Private/ChaosFlesh/Asset/ActorFactoryFlesh.cpp
+++ b/Engine/Plugins/Experimental/ChaosFlesh/Source/ChaosFleshEditor/Private/ChaosFlesh/Asset/ActorFactoryFlesh.cpp
@@ -2,6 +2,7 @@
 
 #include "ChaosFlesh/Asset/ActorFactoryFlesh.h"
 
+#include "ChaosFlesh/ChaosDeformableTetrahedralComponent.h"
 #include "ChaosFlesh/FleshActor.h"
 #include "AssetRegistry/AssetData.h"
 
diff --git a/Engine/Plugins/Experimental/ChaosFlesh/Source/ChaosFleshEditor/Private/ChaosFlesh/Asset/FleshAssetThumbnailScene.cpp b/Engine/Plugins/Experimental/ChaosFlesh/Source/ChaosFleshEditor/Private/ChaosFlesh/Asset/FleshAssetThumbnailScene.cpp
--- a/Engine/Plugins/Experimental/ChaosFlesh/Source/ChaosFleshEditor/Private/ChaosFlesh/Asset/FleshAssetThumbnailScene.cpp
+++ b/Engine/Plugins/Experimental/ChaosFlesh/Source/ChaosFleshEditor/Private/ChaosFlesh/Asset/FleshAssetThumbnailScene.cpp
@@ -1,6 +1,7 @@
 // Copyright Epic Games, Inc. All Rights Reserved.
 
 #include "ChaosFlesh/Asset/FleshAssetThumbnailScene.h"
+#include "ChaosFlesh/ChaosDeformableTetrahedralComponent.h"
 #include "ThumbnailRendering/SceneThumbnailInfo.h"
 #include "ChaosFlesh/FleshActor.h"
 
diff --git a/Engine/Plugins/Experimental/ChaosFlesh/Source/ChaosFleshEngine/Private/ChaosFlesh/ChaosDeformableGPUBuffers.cpp b/Engine/Plugins/Experimental/ChaosFlesh/Source/ChaosFleshEngine/Private/ChaosFlesh/ChaosDeformableGPUBuffers.cpp
new file mode 100644
--- /dev/null
+++ b/Engine/Plugins/Experimental/ChaosFlesh/Source/ChaosFleshEngine/Private/ChaosFlesh/ChaosDeformableGPUBuffers.cpp
@@ -0,0 +1,344 @@
+// Copyright Epic Games, Inc. All Rights Reserved.
+
+#include "ChaosFlesh/ChaosDeformableGPUBuffers.h"
+#include "ChaosFlesh/ChaosDeformableTypes.h"
+
+
+DEFINE_LOG_CATEGORY_STATIC(LogFleshDeformerBuffers, Log, All);
+
+namespace Local
+{
+	FChaosEngineDeformableCVarParams CVarParam;
+	FAutoConsoleVariableRef CVarTestUnpacking(
+		TEXT("p.Chaos.Deformable.FleshDeformer.TestUnpacking"),
+		CVarParam.bTestUnpacking,
+		TEXT("Test flesh deformer data compression. [def: false]"));
+} // namespace Local
+
+using namespace UE::ChaosDeformable;
+
+// Render lingo cheat sheet:
+//  RHI - Render Hardware Interface - C++ interface for D3D/OGL.
+//  SRV - Shader Resource View
+
+//=========================================================================
+// FFloatArrayBufferWithSRV
+//=========================================================================
+
+void FFloatArrayBufferWithSRV::Init(const float* InArray, const int32 Num)
+{
+	Array.SetNumUninitialized(Num);
+	for (int i = 0; i < Num; i++)
+	{
+		Array[i] = InArray[i];
+	}
+	NumValues = Num;
+}
+
+uint32 FFloatArrayBufferWithSRV::GetBufferSize() const
+{
+	return NumValues * sizeof(float);
+}
+
+void FFloatArrayBufferWithSRV::InitRHI()
+{
+    if (Array.Num() > 0)
+    {
+        FRHIResourceCreateInfo CreateInfo(TEXT("UMLDeformerModel::FFloatArrayBufferWithSRV"));
+
+        VertexBufferRHI = RHICreateVertexBuffer(Array.Num() * sizeof(float), BUF_Static | BUF_ShaderResource, CreateInfo);
+        float* Data = reinterpret_cast<float*>(RHILockBuffer(VertexBufferRHI, 0, Array.Num() * sizeof(float), RLM_WriteOnly));
+        for (int32 Index = 0; Index < Array.Num(); ++Index)
+        {
+            Data[Index] = Array[Index];
+        }
+        RHIUnlockBuffer(VertexBufferRHI);
+
+		if (bEmptyArray)
+		{
+			Array.Empty();
+		}
+
+        ShaderResourceViewRHI = RHICreateShaderResourceView(VertexBufferRHI, 4, PF_R32_FLOAT);
+
+		UE_LOG(LogFleshDeformerBuffers, Log,
+			TEXT("FFloatArrayBufferWithSRV - '%s' - Buffer size: %d, Array num: %d"),
+			*BufferName, Array.Num() * sizeof(float), Array.Num());
+	}
+    else
+    {
+        VertexBufferRHI = nullptr;
+        ShaderResourceViewRHI = nullptr;
+    }
+}
+
+//=========================================================================
+// FHalfArrayBufferWithSRV
+//=========================================================================
+
+void FHalfArrayBufferWithSRV::Init(const FFloat16* InArray, const int32 Num)
+{
+    Array.SetNumUninitialized(Num);
+    for (int32 i = 0; i < Num; i++)
+    {
+        Array[i] = InArray[i];
+    }
+	NumValues = Num;
+}
+
+void FHalfArrayBufferWithSRV::Init(const float* InArray, const int32 Num)
+{
+    Array.SetNumUninitialized(Num);
+    for (int32 i = 0; i < Num; i++)
+    {
+        Array[i] = FFloat16(InArray[i]);
+    }
+	NumValues = Num;
+}
+
+void FHalfArrayBufferWithSRV::Init(const TArray<float>& InArray)
+{
+	if (InArray.IsEmpty())
+	{
+		Array.SetNum(0);
+		NumValues = 0;
+	}
+	else
+	{
+		Init(&InArray[0], InArray.Num());
+	}
+}
+
+uint32 FHalfArrayBufferWithSRV::GetBufferSize() const
+{
+	return NumValues * sizeof(FFloat16);
+}
+
+void FHalfArrayBufferWithSRV::InitRHI()
+{
+    if (Array.Num() > 0)
+    {
+        FRHIResourceCreateInfo CreateInfo(TEXT("FHalfArrayBufferWithSRV"));
+
+        VertexBufferRHI = RHICreateVertexBuffer(Array.Num() * sizeof(FFloat16), BUF_Static | BUF_ShaderResource, CreateInfo);
+        FFloat16* Data = reinterpret_cast<FFloat16*>(RHILockBuffer(VertexBufferRHI, 0, Array.Num() * sizeof(FFloat16), RLM_WriteOnly));
+        for (int32 Index = 0; Index < Array.Num(); ++Index)
+        {
+            Data[Index] = Array[Index];
+        }
+        RHIUnlockBuffer(VertexBufferRHI);
+        Array.Empty();
+
+        ShaderResourceViewRHI = RHICreateShaderResourceView(VertexBufferRHI, sizeof(FFloat16), PF_R16F);
+
+		UE_LOG(LogFleshDeformerBuffers, Log,
+			TEXT("FHalfArrayBufferWithSRV - '%s' - Buffer size: %d, Array num: %d"),
+			*BufferName, GetBufferSize(), NumValues);
+	}
+    else
+    {
+        VertexBufferRHI = nullptr;
+        ShaderResourceViewRHI = nullptr;
+    }
+}
+
+//=========================================================================
+// FIndexArrayBufferWithSRV
+//=========================================================================
+
+void FIndexArrayBufferWithSRV::Init(const int32* InArray, const int32 Num)
+{
+	Array.SetNumUninitialized(Num);
+	int32 MinV = TNumericLimits<int32>::Max();
+	int32 MaxV = -TNumericLimits<int32>::Max();
+	for (int32 i = 0; i < Num; i++)
+	{
+		MinV = InArray[i] < MinV ? InArray[i] : MinV;
+		MaxV = InArray[i] > MaxV ? InArray[i] : MaxV;
+		Array[i] = InArray[i];
+	}
+	NumValues = Array.Num();
+
+	int32 Range = MaxV - MinV;
+	bUint8 = Range <= TNumericLimits<uint8>::Max();
+	bUint16 = Range <= TNumericLimits<uint16>::Max();
+	Offset = -MinV; // Adding offset shifts to zero.
+}
+
+void FIndexArrayBufferWithSRV::Init(const uint32* InArray, const int32 Num)
+{
+	Array.SetNumUninitialized(Num);
+	uint32 MinV = TNumericLimits<uint32>::Max();
+	uint32 MaxV = 0;
+	for (int32 i = 0; i < Num; i++)
+	{
+		MinV = InArray[i] < MinV ? InArray[i] : MinV;
+		MaxV = InArray[i] > MaxV ? InArray[i] : MaxV;
+		Array[i] = InArray[i];
+	}
+	NumValues = Array.Num();
+
+	uint32 Range = MaxV - MinV;
+	bUint8 = Range <= TNumericLimits<uint8>::Max();
+	bUint16 = Range <= TNumericLimits<uint16>::Max();
+	Offset = -static_cast<int32>(MinV); // Adding offset shifts to zero.
+}
+
+uint32 FIndexArrayBufferWithSRV::GetBufferSize() const
+{
+	const int32 NumIndices = Array.Num();
+	const int32 ValuesPerSlot = bForce32 ? 1 : bUint8 ? 4 : bUint16 ? 2 : 1;
+	const int32 NumSlots = static_cast<int32>(ceil(static_cast<float>(NumIndices) / ValuesPerSlot));
+	const int32 BufferSize = NumSlots * sizeof(uint32);
+	return BufferSize;
+}
+
+#ifdef WITH_EDITOR
+uint32 UnpackUInt(const uint32* Array, const uint32 Index, const uint32 Stride, const uint32 ArraySize)
+{
+	uint32 ValsPerEntry = 4 / Stride;
+	uint32 BaseIndex = Index / ValsPerEntry;
+	uint32 SubIndex = Index % ValsPerEntry;
+
+	check(BaseIndex < ArraySize);
+
+	if (ValsPerEntry == 1)
+	{
+		check(SubIndex == 0);
+		return Array[BaseIndex];
+	}
+	else if (ValsPerEntry == 2)
+	{
+		const uint16* Array16 = reinterpret_cast<const uint16*>(Array);
+		const uint16 Val16 = Array16[Index]; // the c++ way
+		uint16 Val = (Array[BaseIndex] >> (SubIndex * 16)) & 0xFFFF; // the hlsl way
+		check(Val16 == Val); // make sure they match
+		return Val;
+	}
+ 	else
+	{
+		const uint8* Array8 = reinterpret_cast<const uint8*>(Array);
+		const uint8 Val8 = Array8[Index]; // the c++ way
+		uint8 Val = (Array[BaseIndex] >> (SubIndex * 8)) & 0xFF; // the hlsl way
+		check(Val8 == Val); // make sure they match
+		return Val;
+	}
+}
+#endif
+
+
+FORCEINLINE uint32 GetBits(uint32 Value, uint32 NumBits, uint32 Offset)
+{
+	uint32 Mask = (1u << NumBits) - 1u;
+	return (Value >> Offset) & Mask;
+}
+
+FORCEINLINE void SetBits(uint32& Value, uint32 Bits, uint32 NumBits, uint32 Offset)
+{
+	uint32 Mask = (1u << NumBits) - 1u;
+	check(Bits <= Mask);
+	Mask <<= Offset;
+	Value = (Value & ~Mask) | (Bits << Offset);
+}
+
+void FIndexArrayBufferWithSRV::InitRHI()
+{
+	// Buffer is comprised of 32 bit values.  We're packing 32, 16, or 8 bit ints into 32 bit slots.
+	const int32 NumIndices = Array.Num();
+	const int32 ValuesPerSlot = bForce32 ? 1 : bUint8 ? 4 : bUint16 ? 2 : 1;
+	const int32 NumSlots = static_cast<int32>(ceil(static_cast<float>(NumIndices) / ValuesPerSlot));
+	const int32 BufferSize = NumSlots * sizeof(uint32);
+
+	// From the perspective of the gpu transport system, we're shipping across 32 bits, 
+	// and so wherever that matters, it's hard coded to 'sizeof(uint32)'.
+	const int32 DataStride = 
+		bForce32 ? sizeof(uint32) :
+			bUint8 ? sizeof(uint8) : 
+				bUint16 ? sizeof(uint16) : sizeof(uint32);
+
+    if (BufferSize)
+    {
+        // Create the index buffer.
+        FRHIResourceCreateInfo CreateInfo(*GetFriendlyName());
+		VertexBufferRHI = RHICreateVertexBuffer(BufferSize, BUF_Static | BUF_ShaderResource, CreateInfo);
+
+        // Initialize the buffer.		
+        void* Buffer = RHILockBuffer(VertexBufferRHI, 0, BufferSize, RLM_WriteOnly);
+		
+		if (!bForce32 && bUint8)
+        {
+			// 0: 33 22 11 00
+			// 1: 77 66 55 44
+			// 2: 11 10 99 88
+
+            uint8* DestIndices8Bit = reinterpret_cast<uint8*>(Buffer);
+			int32 i = 0;
+			for (; i < NumIndices; i++)
+            {
+                DestIndices8Bit[i] = static_cast<uint8>(Array[i] + Offset); // add offset, sub in shader
+            }
+			// Zero tailing unused space.
+			for (; i < NumSlots * 4; i++)
+			{
+				DestIndices8Bit[i] = 0;
+			}
+		}
+        else if (!bForce32 && bUint16)
+        {
+			// 0: 1111 0000
+			// 1: 3333 2222
+			// 2: 5555 4444
+			
+			uint16* DestIndices16Bit = reinterpret_cast<uint16*>(Buffer);
+			int32 i = 0;
+			for (; i < NumIndices; i++)
+			{
+				DestIndices16Bit[i] = static_cast<uint16>(Array[i] + Offset); // add offset, sub in shader
+			}
+			// Zero tailing unused space.
+			for (; i < NumSlots * 2; i++)
+			{
+				DestIndices16Bit[i] = 0;
+			}
+		}
+        else // uint32
+        {
+            uint32* DestIndices32Bit = reinterpret_cast<uint32*>(Buffer);
+			for (int32 i = 0; i < NumIndices; i++)
+			{
+				DestIndices32Bit[i] = static_cast<uint32>(Array[i] + Offset); // add offset, sub in shader
+			}
+        }
+
+#ifdef WITH_EDITOR
+		if (Local::CVarParam.bTestUnpacking)
+		{
+			uint32* Indices32Bit = reinterpret_cast<uint32*>(Buffer);
+			for (int32 i = 0; i < NumIndices; i++)
+			{
+				int32 UnpackedInt = static_cast<int>(UnpackUInt(Indices32Bit, i, DataStride, NumIndices)) - Offset; // sub offset for unpacking
+				int32 Orig = Array[i];
+				check(Orig == UnpackedInt);
+			}
+			UE_LOG(LogFleshDeformerBuffers, Log,
+				TEXT("FIndexArrayBufferWithSRV - TEST_UNPACKING succeeded for %d values, packed into buffer size %d, %d bits per value."),
+				NumIndices, BufferSize, (BufferSize / NumIndices));
+		}
+#endif
+
+		RHIUnlockBuffer(VertexBufferRHI);
+        Array.Empty();
+
+		EPixelFormat Format = PF_R32_UINT;
+        ShaderResourceViewRHI = RHICreateShaderResourceView(VertexBufferRHI, sizeof(uint32), Format);
+
+		UE_LOG(LogFleshDeformerBuffers, Log,
+			TEXT("FIndexArrayBufferWithSRV - '%s' - Data stride: %d, Buffer size: %d, Input array size: %d"),
+			*BufferName, DataStride, BufferSize, NumValues);
+	}
+    else
+    {
+        VertexBufferRHI = nullptr;
+        ShaderResourceViewRHI = nullptr;
+    }
+}
diff --git a/Engine/Plugins/Experimental/ChaosFlesh/Source/ChaosFleshEngine/Private/ChaosFlesh/ChaosDeformableSolverComponent.cpp b/Engine/Plugins/Experimental/ChaosFlesh/Source/ChaosFleshEngine/Private/ChaosFlesh/ChaosDeformableSolverComponent.cpp
--- a/Engine/Plugins/Experimental/ChaosFlesh/Source/ChaosFleshEngine/Private/ChaosFlesh/ChaosDeformableSolverComponent.cpp
+++ b/Engine/Plugins/Experimental/ChaosFlesh/Source/ChaosFleshEngine/Private/ChaosFlesh/ChaosDeformableSolverComponent.cpp
@@ -313,7 +313,7 @@ void UDeformableSolverComponent::UpdateFromSimulation(float DeltaTime)
 					{
 						if (const FDataMapValue* Buffer = Output->ObjectMap.Find(DeformableComponent))
 						{
-							DeformableComponent->UpdateFromSimualtion(Buffer);
+							DeformableComponent->UpdateFromSimulation(Buffer);
 						}
 					}
 				}
diff --git a/Engine/Plugins/Experimental/ChaosFlesh/Source/ChaosFleshEngine/Private/ChaosFlesh/ChaosDeformableTetrahedralComponent.cpp b/Engine/Plugins/Experimental/ChaosFlesh/Source/ChaosFleshEngine/Private/ChaosFlesh/ChaosDeformableTetrahedralComponent.cpp
--- a/Engine/Plugins/Experimental/ChaosFlesh/Source/ChaosFleshEngine/Private/ChaosFlesh/ChaosDeformableTetrahedralComponent.cpp
+++ b/Engine/Plugins/Experimental/ChaosFlesh/Source/ChaosFleshEngine/Private/ChaosFlesh/ChaosDeformableTetrahedralComponent.cpp
@@ -28,7 +28,7 @@ FChaosEngineDeformableCVarParams CVarParams;
 FAutoConsoleVariableRef CVarDeforambleDoDrawSimulationMesh(TEXT("p.Chaos.DebugDraw.Deformable.SimulationMesh"), CVarParams.bDoDrawSimulationMesh, TEXT("Debug draw the deformable simulation resutls on the game thread. [def: true]"));
 FAutoConsoleVariableRef CVarDeforambleDoDrawSkeletalMeshBindingPositions(TEXT("p.Chaos.DebugDraw.Deformable.SkeletalMeshBindingPositions"), CVarParams.bDoDrawSkeletalMeshBindingPositions, TEXT("Debug draw the deformable simulation's SkeletalMeshBindingPositions on the game thread. [def: false]"));
 FAutoConsoleVariableRef CVarDeforambleDoDrawSkeletalMeshBindingPositionsSimulationBlendWeight(TEXT("p.Chaos.DebugDraw.Deformable.SkeletalMeshBindingPositions.SimulationBlendWeight"), CVarParams.DrawSkeletalMeshBindingPositionsSimulationBlendWeight, TEXT("Set the simulation blend weight of the skeletal mesh debug draw.[def: 1.]"));
-
+FAutoConsoleVariableRef CVarDeformableFleshDeformerUpdateGPUBuffersOnTick(TEXT("p.Chaos.Deformable.FleshDeformer.UpdateGPUBuffersOnTick"), CVarParams.bUpdateGPUBuffersOnTick, TEXT("Enable/disable time varying updates of GPU buffer data."));
 
 #define PERF_SCOPE(X) SCOPE_CYCLE_COUNTER(X); TRACE_CPUPROFILER_EVENT_SCOPE(X);
 DECLARE_CYCLE_STAT(TEXT("Chaos.Deformable.UDeformableTetrahedralComponent.TickComponent"), STAT_ChaosDeformable_UDeformableTetrahedralComponent_TickComponent, STATGROUP_Chaos);
@@ -39,8 +39,10 @@ DECLARE_CYCLE_STAT(TEXT("Chaos.Deformable.UDeformableTetrahedralComponent.Render
 
 DEFINE_LOG_CATEGORY_STATIC(LogDeformableTetrahedralComponentInternal, Log, All);
 
+
 UDeformableTetrahedralComponent::UDeformableTetrahedralComponent(const FObjectInitializer& ObjectInitializer)
 	: Super(ObjectInitializer)
+	, GPUBufferManager(this)
 {
 	Mesh = ObjectInitializer.CreateDefaultSubobject<UProceduralMeshComponent>(this, TEXT("Flesh Visualization Component"));
 	PrimaryComponentTick.TickGroup = TG_LastDemotable;
@@ -61,7 +63,6 @@ void UDeformableTetrahedralComponent::Invalidate()
 	bBoundsNeedsUpdate = true;
 }
 
-
 void UDeformableTetrahedralComponent::OnRegister()
 {
 	if (bBoundsNeedsUpdate)
@@ -207,7 +208,7 @@ UDeformablePhysicsComponent::FDataMapValue UDeformableTetrahedralComponent::NewD
 								*GetSimulationCollection()->GetCollection(),
 								this->GetComponentTransform(),
 								BoneSpaceXf, 
-							SimulationSpace.SimSpaceTransformGlobalIndex,
+								SimulationSpace.SimSpaceTransformGlobalIndex,
 								AnimationTransforms, 
 								ComponentPose, 
 								BodyForces.bApplyGravity,
@@ -227,7 +228,7 @@ UDeformablePhysicsComponent::FDataMapValue UDeformableTetrahedralComponent::NewD
 			*GetSimulationCollection()->GetCollection(),
 			this->GetComponentTransform(),
 			GetSimSpaceRestTransform(),
-		SimulationSpace.SimSpaceTransformGlobalIndex,
+			SimulationSpace.SimSpaceTransformGlobalIndex,
 			BodyForces.bApplyGravity,
 			BodyForces.StiffnessMultiplier,
 			BodyForces.DampingMultiplier,
@@ -304,7 +305,7 @@ FTransform UDeformableTetrahedralComponent::GetSimSpaceRestTransform() const
 	return ComponentToBone;
 }
 
-void UDeformableTetrahedralComponent::UpdateFromSimualtion(const FDataMapValue* SimualtionBuffer)
+void UDeformableTetrahedralComponent::UpdateFromSimulation(const FDataMapValue* SimualtionBuffer)
 {
 	PERF_SCOPE(STAT_ChaosDeformable_UDeformableTetrahedralComponent_UpdateFromSimualtion);
 
@@ -325,8 +326,15 @@ void UDeformableTetrahedralComponent::UpdateFromSimualtion(const FDataMapValue*
 				DynamicVertex[i] = SimulationVertex[i];
 			}
 			
-			//p.Chaos.DebugDraw.Enabled 1
-			//p.Chaos.DebugDraw.Deformable.SkeletalMeshBindingPositions 1
+			// p.Chaos.Deformable.FleshDeformer.UpdateGPUBuffersOnTick 1 (default) or 0
+			if (CVarParams.bUpdateGPUBuffersOnTick)
+			{
+				// Update time varying GPU buffers (but only if a consumer has been registered).
+				GPUBufferManager.UpdateGPUBuffers();
+			}
+
+			// p.Chaos.DebugDraw.Enabled 1
+			// p.Chaos.DebugDraw.Deformable.SkeletalMeshBindingPositions 1
 			if (CVarParams.bDoDrawSkeletalMeshBindingPositions)
 			{
 				DebugDrawSkeletalMeshBindingPositions();
@@ -939,5 +947,3 @@ TArray<FVector> UDeformableTetrahedralComponent::GetSkeletalMeshBindingPositions
 	}
 	return TransformPositions;
 }
-
-
diff --git a/Engine/Plugins/Experimental/ChaosFlesh/Source/ChaosFleshEngine/Private/ChaosFlesh/ChaosFleshDeformerBufferManager.cpp b/Engine/Plugins/Experimental/ChaosFlesh/Source/ChaosFleshEngine/Private/ChaosFlesh/ChaosFleshDeformerBufferManager.cpp
new file mode 100644
--- /dev/null
+++ b/Engine/Plugins/Experimental/ChaosFlesh/Source/ChaosFleshEngine/Private/ChaosFlesh/ChaosFleshDeformerBufferManager.cpp
@@ -0,0 +1,756 @@
+// Copyright Epic Games, Inc. All Rights Reserved.
+
+#include "ChaosFlesh/ChaosFleshDeformerBufferManager.h"
+#include "ChaosFlesh/FleshCollection.h"
+#include "ChaosFlesh/ChaosDeformableTetrahedralComponent.h"
+#include "ChaosFlesh/ChaosDeformableTypes.h"
+#include "ChaosFlesh/FleshDynamicAsset.h"
+#include "ChaosFlesh/SimulationAsset.h"
+#include "GeometryCollection/Facades/CollectionTetrahedralBindingsFacade.h"
+
+#include "RenderResource.h"
+
+DEFINE_LOG_CATEGORY_STATIC(LogFleshDeformerBufferManager, Log, All);
+
+namespace Local 
+{
+	FChaosEngineDeformableCVarParams CVarParams;
+	FAutoConsoleVariableRef CVarTestSparseMappings(
+		TEXT("p.Chaos.Deformable.FleshDeformer.TestSparseMappings"),
+		CVarParams.bTestSparseMappings,
+		TEXT("Test flesh deformer sparsification."));
+	FAutoConsoleVariableRef CVarTestCompressedIndexing(
+		TEXT("p.Chaos.Deformable.FleshDeformer.TestCompressedIndexing"),
+		CVarParams.bTestCompressedIndexing,
+		TEXT("Test flesh deformer compressed index encodings."));
+} // namespace Local
+
+using namespace Chaos::Softs;
+
+//=============================================================================
+// Util
+//=============================================================================
+
+namespace Chaos::Softs::Private
+{
+	TArray<float> NullFloatBuffer;
+	TArray<FFloat16> NullHalfBuffer;
+	TArray<int32> NullIndexBuffer;
+
+	template<class T_Buffer, class T_StoreData>
+	void InitBuffer(
+		T_Buffer& Buffer,
+		const TArray<FIntVector4>* Array,
+		const TArray<T_StoreData>& NullBuffer)
+	{
+		if (Buffer.IsInitialized())
+		{
+			BeginReleaseResource(&Buffer);
+		}
+		if (Array && Array->Num())
+		{
+			Buffer.Init(Array->GetData(), Array->Num());
+		}
+		else
+		{
+			Buffer.Init(NullBuffer);
+		}
+		BeginInitResource(&Buffer);
+	}
+
+	template<class T_Buffer, class T_StoreData>
+	void InitBuffer(
+		T_Buffer& Buffer,
+		const TArray<uint32>* Array,
+		const TArray<T_StoreData>& NullBuffer)
+	{
+		if (Buffer.IsInitialized())
+		{
+			BeginReleaseResource(&Buffer);
+		}
+		if (Array && Array->Num())
+		{
+			const uint32* Data = Array->GetData();
+			Buffer.Init(Data, Array->Num());
+		}
+		else
+		{
+			Buffer.Init(NullBuffer);
+		}
+		BeginInitResource(&Buffer);
+	}
+
+	template<class T_Buffer, class T_StoreData>
+	void InitBuffer(
+		T_Buffer& Buffer,
+		const TArray<float>* Array,
+		const TArray<T_StoreData>& NullBuffer)
+	{
+		if (Buffer.IsInitialized())
+		{
+			BeginReleaseResource(&Buffer);
+		}
+		if (Array && Array->Num())
+		{
+			const float* Data = Array->GetData();
+			Buffer.Init(Data, Array->Num());
+		}
+		else
+		{
+			Buffer.Init(NullBuffer);
+		}
+		BeginInitResource(&Buffer);
+	}
+
+	template<class T_Buffer, class T_StoreData>
+	void InitBuffer(
+		T_Buffer& Buffer,
+		const TArray<FVector3f>* Array,
+		const TArray<T_StoreData>& NullBuffer)
+	{
+		if (Buffer.IsInitialized())
+		{
+			BeginReleaseResource(&Buffer);
+		}
+		if (Array && Array->Num())
+		{
+			FVector3f* Data = const_cast<FVector3f*>(Array->GetData());
+			Buffer.Init(&(*Data)[0], Array->Num() * 3);
+		}
+		else
+		{
+			Buffer.Init(NullBuffer);
+		}
+		BeginInitResource(&Buffer);
+	}
+
+	template<class T_Buffer, class T_StoreData>
+	void InitBuffer(
+		T_Buffer& Buffer,
+		const TArray<FVector4f>* Array,
+		const TArray<T_StoreData>& NullBuffer)
+	{
+		if (Buffer.IsInitialized())
+		{
+			BeginReleaseResource(&Buffer);
+		}
+		if (Array && Array->Num())
+		{
+			FVector4f* Data = const_cast<FVector4f*>(Array->GetData());
+			Buffer.Init(&(*Data)[0], Array->Num() * 4);
+		}
+		else
+		{
+			Buffer.Init(NullBuffer);
+		}
+		BeginInitResource(&Buffer);
+	}
+
+	template <class T_Buffer, class T_Data, class T_StoreData>
+	void InitBuffer(
+		T_Buffer& Buffer,
+		const TManagedArrayAccessor<T_Data>* ManagedArray,
+		const TArray<T_StoreData>& NullBuffer)
+	{
+		const TArray<T_Data>* Array = ManagedArray ? &ManagedArray->Get().GetConstArray() : nullptr;
+		return InitBuffer(Buffer, Array, NullBuffer);
+	}
+
+	template<class T_Buffer, class T_StoreData>
+	void UpdateBuffer(
+		T_Buffer& Buffer,
+		const TArray<FVector3f>* Array,
+		const TArray<T_StoreData>& NullBuffer)
+	{
+		check(Buffer.IsInitialized());
+		if (Array && Array->Num())
+		{
+			FVector3f* Data3f = const_cast<FVector3f*>(Array->GetData());
+			Buffer.Init(&(*Data3f)[0], Array->Num() * 3);
+		}
+		else
+		{
+			Buffer.Init(NullBuffer);
+		}
+		BeginUpdateResourceRHI(&Buffer);
+	}
+
+} // namespace Chaos::Softs::Private
+
+
+//=============================================================================
+// FChaosFleshDeformableGPUManager::BindingsBuffer
+//=============================================================================
+
+FChaosFleshDeformableGPUManager::FBindingsBuffer::FBindingsBuffer()
+{
+	VerticesBuffer.SetBufferName(FString("FleshDeformer.SparseTetVertices"));
+	RestVerticesBuffer.SetBufferName(FString("FleshDeformer.RestVertices"));
+
+	ParentsBuffer.SetBufferName(FString("FleshDeformer.TetParents"));
+	VertsToParentsBuffer.SetBufferName(FString("FleshDeformer.VertsToParents"));
+
+	WeightsBuffer.SetBufferName(FString("FleshDeformer.TetWeights"));
+	OffsetsBuffer.SetBufferName(FString("FleshDeformer.SurfaceOffsets"));
+	MaskBuffer.SetBufferName(FString("FleshDeformer.Mask"));
+
+	// The vertices buffer should retain its local memory, as it'll be updated on tick.
+	VerticesBuffer.RetainLocalMemory();
+}
+
+FChaosFleshDeformableGPUManager::FBindingsBuffer::~FBindingsBuffer()
+{
+	DestroyGPUBuffers();
+}
+
+void FChaosFleshDeformableGPUManager::FBindingsBuffer::DestroyGPUBuffers()
+{
+	//StartBatchedRelease();
+	if (VerticesBuffer.IsInitialized())
+	{
+		BeginReleaseResource(&VerticesBuffer);
+	}
+
+	if (RestVerticesBuffer.IsInitialized())
+	{
+		BeginReleaseResource(&RestVerticesBuffer);
+	}
+	if (ParentsBuffer.IsInitialized())
+	{
+		BeginReleaseResource(&ParentsBuffer);
+	}
+	if (VertsToParentsBuffer.IsInitialized())
+	{
+		BeginReleaseResource(&VertsToParentsBuffer);
+	}
+	if (WeightsBuffer.IsInitialized())
+	{
+		BeginReleaseResource(&WeightsBuffer);
+	}
+	if (OffsetsBuffer.IsInitialized())
+	{
+		BeginReleaseResource(&OffsetsBuffer);
+	}
+	if (MaskBuffer.IsInitialized())
+	{
+		BeginReleaseResource(&MaskBuffer);
+	}
+	//EndBatchedRelease();
+	RenderResourceDestroyFence.BeginFence();
+	RenderResourceDestroyFence.Wait();
+}
+
+void FChaosFleshDeformableGPUManager::FBindingsBuffer::SetOwnerName(const FName& OwnerName)
+{
+	VerticesBuffer.SetOwnerName(OwnerName);
+	RestVerticesBuffer.SetOwnerName(OwnerName);
+
+	ParentsBuffer.SetOwnerName(OwnerName);
+	VertsToParentsBuffer.SetOwnerName(OwnerName);
+	WeightsBuffer.SetOwnerName(OwnerName);
+	OffsetsBuffer.SetOwnerName(OwnerName);
+	MaskBuffer.SetOwnerName(OwnerName);
+}
+
+bool FChaosFleshDeformableGPUManager::FBindingsBuffer::InitData(
+	const TManagedArrayAccessor<FIntVector4>* ParentsArray,
+	const TManagedArrayAccessor<FVector4f>* WeightsArray,
+	const TManagedArrayAccessor<FVector3f>* OffsetArray,
+	const TManagedArrayAccessor<float>* MaskArray,
+	const TManagedArray<FVector3f>* RestVerticesArray,
+	const TManagedArray<FVector3f>* VerticesArray)
+{
+	InitCondensationMapping(*ParentsArray);
+
+	// Bindings
+	Chaos::Softs::Private::InitBuffer(VertsToParentsBuffer, &VertsToSparseParents, Chaos::Softs::Private::NullIndexBuffer);
+	Chaos::Softs::Private::InitBuffer(ParentsBuffer, &SparseParents, Chaos::Softs::Private::NullIndexBuffer);
+	Chaos::Softs::Private::InitBuffer(WeightsBuffer, WeightsArray, Chaos::Softs::Private::NullHalfBuffer);
+	Chaos::Softs::Private::InitBuffer(OffsetsBuffer, OffsetArray, Chaos::Softs::Private::NullHalfBuffer);
+	Chaos::Softs::Private::InitBuffer(MaskBuffer, MaskArray, Chaos::Softs::Private::NullHalfBuffer);
+
+	// Rest vertices
+	const TArray<FVector3f>& DenseRestVertex = RestVerticesArray->GetConstArray();
+	UpdateSparseVertices(DenseRestVertex, SparseRestVertices);
+	Chaos::Softs::Private::InitBuffer(RestVerticesBuffer, &SparseRestVertices, Chaos::Softs::Private::NullFloatBuffer);
+
+	// Dynamic vertices
+	const TArray<FVector3f>& DenseVertex = VerticesArray->GetConstArray();
+	if (DenseVertex.Num() == DenseRestVertex.Num())
+	{
+		UpdateSparseVertices(DenseVertex, SparseVertices);
+	}
+	else
+	{
+		UpdateSparseVertices(DenseRestVertex, SparseVertices);
+	}
+	Chaos::Softs::Private::InitBuffer(VerticesBuffer, &SparseVertices, Chaos::Softs::Private::NullFloatBuffer);
+
+#ifdef WITH_EDITOR
+	if (Local::CVarParams.bTestSparseMappings)
+	{
+		TMap<int32, int32> DenseToSparseIndices;
+		DenseToSparseIndices.Reserve(SparseToDenseIndices.Num());
+		for (int32 i = 0; i < SparseToDenseIndices.Num(); i++)
+		{
+			DenseToSparseIndices.Add(TTuple<int32, int32>(SparseToDenseIndices[i], i));
+		}
+		check(DenseToSparseIndices.Num() == SparseToDenseIndices.Num());
+
+		const TArray<FIntVector4>& DenseParents = ParentsArray->Get().GetConstArray();
+		for (int32 i = 0; i < VertsToSparseParents.Num(); i++)
+		{
+			// Sparse parent index for ender index i
+			const uint32 SparseParentIdx = VertsToSparseParents[i];
+			check(SparseParents.IsValidIndex(SparseParentIdx));
+			const FIntVector4& SparseParent = SparseParents[SparseParentIdx];
+
+			// Remap sparse parents back to dense indices
+			FIntVector4 Parents;
+			for (int32 j = 0; j < 4; j++)
+			{
+				const int32 SparseIndex = SparseParent[j];
+				if (SparseIndex == INDEX_NONE)
+					Parents[j] = INDEX_NONE;
+				else
+				{
+					check(SparseToDenseIndices.IsValidIndex(SparseIndex));
+					Parents[j] = SparseToDenseIndices[SparseIndex];
+				}
+			}
+
+			// Check that the unmapped sparse indices match the dense indices.
+			const FIntVector4& DenseParent = DenseParents[i];
+			check(Parents == DenseParent);
+
+			// Check that the sparse vertices match the dense vertices.
+			if (DenseVertex.Num() == DenseRestVertex.Num())
+			{
+				for (int32 j = 0; j < 4; j++)
+					if (SparseParent[j] == INDEX_NONE || DenseParent[j] == INDEX_NONE)
+					{
+						check(SparseParent[j] == DenseParent[j]);
+					}
+					else
+					{
+						check(SparseVertices[SparseParent[j]] == DenseVertex[DenseParent[j]]);
+					}
+			}
+			else
+			{
+				for (int32 j = 0; j < 4; j++)
+					if (SparseParent[j] == INDEX_NONE || DenseParent[j] == INDEX_NONE)
+					{
+						check(SparseParent[j] == DenseParent[j]);
+					}
+					else
+					{
+						check(SparseVertices[SparseParent[j]] == DenseRestVertex[DenseParent[j]]);
+					}
+			}
+		}
+		UE_LOG(LogFleshDeformerBufferManager, Log,
+			TEXT("FChaosFleshDeformableGPUManager::FBindingsBuffer::InitData() - TEST_SPARSE_MAPPINGS succeeded."));
+	}
+	if(Local::CVarParams.bTestCompressedIndexing)
+	{
+		TArray<int32> DataRange32;
+		TArray<int32> DataRange16;
+		TArray<int32> DataRange8;
+		DataRange32.SetNum(1024);
+		DataRange16.SetNum(1024);
+		DataRange8.SetNum(1024);
+		for (int i = 0; i < 1024; i++)
+		{
+			DataRange32[i] = FMath::RandRange(INDEX_NONE, TNumericLimits<uint16>::Max()*2);
+			DataRange16[i] = FMath::RandRange(INDEX_NONE, TNumericLimits<uint16>::Max()-1);
+			DataRange8[i] = FMath::RandRange(INDEX_NONE, TNumericLimits<uint8>::Max()-1);
+		}
+		{
+			UE::ChaosDeformable::FIndexArrayBufferWithSRV Buffer32; 
+			Buffer32.SetBufferName("Test_int32"); 
+			Buffer32.Init(DataRange32.GetData(), DataRange32.Num()); 
+			Buffer32.InitRHI();
+		}
+		{
+			UE::ChaosDeformable::FIndexArrayBufferWithSRV Buffer16; 
+			Buffer16.SetBufferName("Test_int16"); 
+			Buffer16.Init(DataRange16.GetData(), DataRange16.Num()); 
+			Buffer16.InitRHI();
+		}
+		{
+			UE::ChaosDeformable::FIndexArrayBufferWithSRV Buffer8; 
+			Buffer8.SetBufferName("Test_int8"); 
+			Buffer8.Init(DataRange8.GetData(), DataRange8.Num()); 
+			Buffer8.InitRHI();
+		}
+	}
+#endif
+
+	UE_LOG(LogFleshDeformerBufferManager, Log,
+		TEXT("FChaosFleshDeformableGPUManager::FBindingsBuffer::InitData() - Data size stats:\n"
+			"    VertsToParentsBuffer: %d\n"
+			"    ParentsBuffer: %d\n"
+			"    WeightsBuffer: %d\n"
+			"    OffsetsBuffer: %d\n"
+			"    MaskBuffer:    %d\n"
+			"    RestVertex:    %d\n"
+			"    VertexBuffer:  %d\n"
+			"    TOTAL:         %d"),
+		VertsToParentsBuffer.GetBufferSize(),
+		ParentsBuffer.GetBufferSize(),
+		WeightsBuffer.GetBufferSize(),
+		OffsetsBuffer.GetBufferSize(),
+		MaskBuffer.GetBufferSize(),
+		RestVerticesBuffer.GetBufferSize(),
+		VerticesBuffer.GetBufferSize(),
+
+		VertsToParentsBuffer.GetBufferSize() +
+		ParentsBuffer.GetBufferSize() +
+		WeightsBuffer.GetBufferSize() +
+		OffsetsBuffer.GetBufferSize() +
+		MaskBuffer.GetBufferSize() +
+		RestVerticesBuffer.GetBufferSize() +
+		VerticesBuffer.GetBufferSize());
+
+	return true;
+}
+
+bool FChaosFleshDeformableGPUManager::FBindingsBuffer::UpdateVertices(const TArray<FVector3f>& DenseVertex)
+{
+	if (UpdateSparseVertices(DenseVertex, SparseVertices))
+	{
+		Chaos::Softs::Private::UpdateBuffer(VerticesBuffer, &SparseVertices, Chaos::Softs::Private::NullFloatBuffer);
+		return true;
+	}
+	return false;
+}
+
+bool FChaosFleshDeformableGPUManager::FBindingsBuffer::UpdateSparseVertices(
+	const TArray<FVector3f>& DenseVertexIn, 
+	TArray<FVector3f>& SparseVerticesIn) const
+{
+	SparseVerticesIn.SetNumUninitialized(SparseToDenseIndices.Num());
+	bool bChanged = false;
+	for (int32 i = 0; i < SparseToDenseIndices.Num(); i++)
+	{
+		if (!DenseVertexIn.IsValidIndex(SparseToDenseIndices[i]))
+		{
+			return bChanged;
+		}
+		bChanged |= SparseVerticesIn[i] != DenseVertexIn[SparseToDenseIndices[i]];
+		SparseVerticesIn[i] = DenseVertexIn[SparseToDenseIndices[i]];
+	}
+	return bChanged;
+}
+
+void FChaosFleshDeformableGPUManager::FBindingsBuffer::InitCondensationMapping(
+	const TManagedArrayAccessor<FIntVector4>& ParentsArray)
+{
+	// We want to skip vertices not involved in the deformation.  The Parents array contains all referenced
+	// nodes.  Construct a condensation mapping from that.
+	TSet<int32> Visited;
+	const TArray<FIntVector4>& DenseParents = ParentsArray.Get().GetConstArray();
+	for (int32 i = 0; i < DenseParents.Num(); i++)
+	{
+		const FIntVector4& Tet = DenseParents[i];
+		for (int32 j = 0; j < 4; j++)
+		{
+			if (Tet[j] >= -1)
+			{
+				Visited.Add(Tet[j]);
+			}
+		}
+	}
+	// Flatten Visited set into a sorted array.
+	TArray<int32> VisitedArray;
+	VisitedArray.Reserve(Visited.Num());
+	for (TSet<int32>::TConstIterator It = Visited.CreateConstIterator(); It; ++It)
+	{
+		const int32 Idx = *It;
+		ensure(Idx >= -1);
+		VisitedArray.Add(Idx);
+	}
+	VisitedArray.Sort();
+
+	// CondensationMapping remaps sparse vertex indices to dense indices.  Every vertex that was visited
+	// gets added to the list, the index of which is the new "SparseIndex".  DenseToSparseIndices is the
+	// inverse of CondensationMapping.
+	SparseToDenseIndices.Empty();
+	SparseToDenseIndices.Reserve(VisitedArray.Num());
+	TMap<int32, int32> DenseToSparseIndices;
+	DenseToSparseIndices.Reserve(VisitedArray.Num());
+	DenseToSparseIndices.Add(INDEX_NONE, INDEX_NONE);
+	for (int32 i = 0; i < VisitedArray.Num(); i++)
+	{
+		if (VisitedArray[i] >= 0) // skip -1, which exists in surface bindings
+		{
+			const int32 SparseIndex = SparseToDenseIndices.Add(VisitedArray[i]);
+			DenseToSparseIndices.Add(VisitedArray[i], SparseIndex);
+		}
+	}
+
+	// The Parents array likely has a lot of duplicate entries, so we consolidate Parents into SparseParents,
+	// and add a level of indirection with VertsToSparseParents.  That way we store 1 index per render vertex
+	// rather than 4.
+	VertsToSparseParents.SetNum(DenseParents.Num());
+	SparseParents.Reserve(DenseParents.Num() / 4);
+	for (int32 i = 0; i < DenseParents.Num(); i++)
+	{
+		FIntVector4 Tet = DenseParents[i];
+		// Remap non-sparse parent vertex indices to sparse vertex indices.
+		for (int32 j = 0; j < 4; j++)
+		{
+			Tet[j] = DenseToSparseIndices[Tet[j]];
+			ensure(Tet[j] >= -1);
+		}
+		int32 Idx = SparseParents.AddUnique(Tet);
+		VertsToSparseParents[i] = Idx;
+	}
+
+	{
+		float SparseToDenseRatio = static_cast<float>(SparseParents.Num() * 4 + VertsToSparseParents.Num()) / (DenseParents.Num() * 4);
+		float SparseToDensePct = (1.0 - SparseToDenseRatio) * 100.0;
+
+		UE_LOG(LogFleshDeformerBufferManager, Log,
+			TEXT("BindingsBuffer::InitCondensationMapping() - DenseParents: %d (%d); SparseParents: %d (%d) + VertsToSparseParents: %d = %d; %g pct reduction."),
+			DenseParents.Num(), (DenseParents.Num() * 4),
+			SparseParents.Num(), (SparseParents.Num() * 4), VertsToSparseParents.Num(),
+			(SparseParents.Num() * 4 + VertsToSparseParents.Num()),
+			SparseToDensePct);
+	}
+
+}
+
+//=============================================================================
+// FChaosFleshDeformableGPUManager
+//=============================================================================
+
+FChaosFleshDeformableGPUManager::FChaosFleshDeformableGPUManager()
+	: Component(nullptr)
+{}
+
+FChaosFleshDeformableGPUManager::FChaosFleshDeformableGPUManager(UDeformableTetrahedralComponent* InComponent)
+	: Component(InComponent)
+{}
+
+FChaosFleshDeformableGPUManager::~FChaosFleshDeformableGPUManager()
+{
+	GPUBufferConsumers.Empty();
+	BufferTable.Empty();
+}
+
+void FChaosFleshDeformableGPUManager::SetOwner(UDeformableTetrahedralComponent* InComponent)
+{
+	Component = InComponent;
+}
+
+void FChaosFleshDeformableGPUManager::UnRegisterGPUBufferConsumer(const void* ID) const 
+{ 
+	if (GPUBufferConsumers.Num())
+	{
+		GPUBufferConsumers.Remove(ID); // mutable
+
+		FChaosFleshDeformableGPUManager* NCThis = const_cast<FChaosFleshDeformableGPUManager*>(this);
+		for (IdToBuffers::TIterator It = NCThis->BufferTable.CreateIterator(); It;)
+		{
+			if (It->Key.Get<0>() == ID)
+			{
+				IdToBuffers::TIterator TargetIt = It;
+				++It;
+				NCThis->BufferTable.Remove(TargetIt->Key);
+			}
+			else
+			{
+				++It;
+			}
+		}
+	}
+}
+
+bool FChaosFleshDeformableGPUManager::HasGPUBindingsBuffer(
+	const void* ID,
+	const FName MeshName,
+	const int32 LodIndex) const 
+{
+	TTuple<void*, FName, int32> Key(const_cast<void*>(ID), MeshName, LodIndex);
+	return BufferTable.Contains(Key);
+}
+
+bool FChaosFleshDeformableGPUManager::InitGPUBindingsBuffer(
+	const void* ID, 
+	const FName MeshName, 
+	const int32 LodIndex,
+	const bool ForceInit)
+{
+	if (!ForceInit && HasGPUBindingsBuffer(ID, MeshName, LodIndex))
+	{
+		return true;
+	}
+
+	if (!Component)
+	{
+		return false;
+	}
+
+	const FFleshCollection* RestCollection = GetRestCollection();
+	if (!RestCollection)
+	{
+		return false;
+	}
+
+	GeometryCollection::Facades::FTetrahedralBindings TetBindings(*RestCollection);
+	const int32 TetIndex = TetBindings.GetTetMeshIndex(MeshName, LodIndex);
+	if (TetIndex != INDEX_NONE)
+	{
+		if (!TetBindings.ReadBindingsGroup(TetIndex, MeshName, LodIndex))
+		{
+			UE_LOG(LogFleshDeformerBufferManager, Error,
+				TEXT("LogFleshDeformerBufferManager: Failed to read bindings group associated with mesh '%s' LOD: %d"),
+				*MeshName.ToString(), LodIndex);
+			return false;
+		}
+		if (!TetBindings.IsValid())
+		{
+			UE_LOG(LogFleshDeformerBufferManager, Error,
+				TEXT("LogFleshDeformerBufferManager: Bindings group associated with mesh '%s' LOD: %d are invalid."),
+				*MeshName.ToString(), LodIndex);
+			return false;
+		}
+	}
+	else 
+	{
+		UE_LOG(LogFleshDeformerBufferManager, Error,
+			TEXT("LogFleshDeformerBufferManager: Failed to find tet mesh index associated to mesh '%s' LOD: %d"),
+			*MeshName.ToString(), LodIndex);
+		return false;
+	}
+
+	const TManagedArrayAccessor<FIntVector4>* ParentsArray = TetBindings.GetParentsRO();
+	const TManagedArrayAccessor<FVector4f>* WeightsArray = TetBindings.GetWeightsRO();
+	const TManagedArrayAccessor<FVector3f>* OffsetArray = TetBindings.GetOffsetsRO();
+	const TManagedArrayAccessor<float>* MaskArray = TetBindings.GetMaskRO();
+	if (!ParentsArray || !WeightsArray || !OffsetArray || !MaskArray)
+	{
+		return false;
+	}
+
+	const TManagedArray<FVector3f>* RestVertexArray = GetRestVertexArray();
+	const TManagedArray<FVector3f>* VertexArray = GetVertexArray();
+	if (!RestVertexArray || !VertexArray)
+	{
+		return false;
+	}
+	
+	TSharedPtr<FBindingsBuffer> Bindings(new FBindingsBuffer());
+	Bindings->SetOwnerName(FName(Component->GetName()));
+	if (Bindings->InitData(
+		ParentsArray, WeightsArray, OffsetArray, MaskArray,
+		RestVertexArray, VertexArray))
+	{
+		TTuple<void*, FName, int32> Key(const_cast<void*>(ID), MeshName, LodIndex);
+		BufferTable.Add(MoveTemp(Key), MoveTemp(Bindings));
+		RegisterGPUBufferConsumer(ID);
+		return true;
+	}
+	return false;
+}
+
+bool FChaosFleshDeformableGPUManager::UpdateGPUBuffers()
+{
+	if (!Component || !HasRegisteredGPUBufferConsumer())
+	{
+		return false;
+	}
+
+	const TManagedArray<FVector3f>* DenseVertexArray = GetVertexArray();
+	if (!DenseVertexArray)
+	{
+		return false;
+	}
+	const TArray<FVector3f>& DenseVertex = DenseVertexArray->GetConstArray();
+
+	for(IdToBuffers::TIterator It = BufferTable.CreateIterator(); It; ++It)
+	{
+		if (TSharedPtr<FBindingsBuffer>& Buffer = It.Value())
+		{
+			Buffer->UpdateVertices(DenseVertex);
+		}
+	}
+	return !BufferTable.IsEmpty();
+}
+
+const FChaosFleshDeformableGPUManager::FBindingsBuffer* FChaosFleshDeformableGPUManager::GetBindingsBuffer(
+	const void* ID, 
+	const FName MeshName, 
+	const int32 LodIndex) const
+{
+	TTuple<void*, FName, int32> Key(const_cast<void*>(ID), MeshName, LodIndex);
+	const TSharedPtr<FBindingsBuffer>* Bindings = BufferTable.Find(Key);
+	return Bindings ? Bindings->Get() : nullptr;
+}
+
+FChaosFleshDeformableGPUManager::FBindingsBuffer* FChaosFleshDeformableGPUManager::GetBindingsBuffer(
+	const void* ID,
+	const FName MeshName,
+	const int32 LodIndex)
+{
+	TTuple<void*, FName, int32> Key(const_cast<void*>(ID), MeshName, LodIndex);
+	TSharedPtr<FBindingsBuffer>* Bindings = BufferTable.Find(Key);
+	return Bindings ? Bindings->Get() : nullptr;
+}
+
+const FFleshCollection* FChaosFleshDeformableGPUManager::GetRestCollection() const
+{
+	if (!Component)
+	{
+		return nullptr;
+	}
+
+	const UFleshAsset* RestCollectionAsset = Component->GetRestCollection();
+	if (!RestCollectionAsset)
+	{
+		return nullptr;
+	}
+
+	const FFleshCollection* RestCollection = RestCollectionAsset->GetCollection();
+	return RestCollection; // can be null
+}
+
+const TManagedArray<FVector3f>* FChaosFleshDeformableGPUManager::GetRestVertexArray() const
+{
+	if (!Component)
+	{
+		return nullptr;
+	}
+
+	if (const UFleshAsset* RestCollection = Component->GetRestCollection())
+	{
+		return RestCollection->FindPositions();
+	}
+	return nullptr;
+}
+
+const TManagedArray<FVector3f>* FChaosFleshDeformableGPUManager::GetVertexArray() const
+{
+	if (!Component)
+	{
+		return nullptr;
+	}
+	if (UFleshDynamicAsset* DynamicCollection = Component->GetDynamicCollection())
+	{
+		if (const UFleshAsset* RestCollection = Component->GetRestCollection())
+		{
+			const int32 NumRestVertices = DynamicCollection->GetCollection()->NumElements(FGeometryCollection::VerticesGroup);
+			const int32 NumDynVertices = DynamicCollection->GetCollection()->NumElements(FGeometryCollection::VerticesGroup);
+			if (NumRestVertices == NumDynVertices)
+			{
+				const TManagedArray<FVector3f>* RetVal = DynamicCollection->FindPositions();
+				return RetVal;
+			}
+		}
+	}
+
+	return GetRestVertexArray();
+}
+
diff --git a/Engine/Plugins/Experimental/ChaosFlesh/Source/ChaosFleshEngine/Private/ChaosFlesh/FleshActor.cpp b/Engine/Plugins/Experimental/ChaosFlesh/Source/ChaosFleshEngine/Private/ChaosFlesh/FleshActor.cpp
--- a/Engine/Plugins/Experimental/ChaosFlesh/Source/ChaosFleshEngine/Private/ChaosFlesh/FleshActor.cpp
+++ b/Engine/Plugins/Experimental/ChaosFlesh/Source/ChaosFleshEngine/Private/ChaosFlesh/FleshActor.cpp
@@ -4,6 +4,7 @@
 
 #include "ChaosFlesh/ChaosDeformableSolverComponent.h"
 #include "ChaosFlesh/ChaosDeformableSolverActor.h"
+#include "ChaosFlesh/ChaosDeformableTetrahedralComponent.h"
 
 #include UE_INLINE_GENERATED_CPP_BY_NAME(FleshActor)
 
diff --git a/Engine/Plugins/Experimental/ChaosFlesh/Source/ChaosFleshEngine/Private/DataInterfaces/DIFleshDeformer.cpp b/Engine/Plugins/Experimental/ChaosFlesh/Source/ChaosFleshEngine/Private/DataInterfaces/DIFleshDeformer.cpp
--- a/Engine/Plugins/Experimental/ChaosFlesh/Source/ChaosFleshEngine/Private/DataInterfaces/DIFleshDeformer.cpp
+++ b/Engine/Plugins/Experimental/ChaosFlesh/Source/ChaosFleshEngine/Private/DataInterfaces/DIFleshDeformer.cpp
@@ -2,6 +2,7 @@
 
 #include "DataInterfaces/DIFleshDeformer.h"
 
+#include "ChaosFlesh/ChaosFleshDeformerBufferManager.h"
 #include "Components/SkeletalMeshComponent.h"
 #include "Containers/Set.h"
 #include "ComputeFramework/ComputeKernelPermutationVector.h"
@@ -24,6 +25,14 @@ DEFINE_LOG_CATEGORY(LogFleshDeformer);
 // Interface
 //
 
+UDIFleshDeformer::~UDIFleshDeformer()
+{
+	if (UDeformableTetrahedralComponent* Component = ProducerComponent.Get())
+	{
+		Component->GetGPUBufferManager().UnRegisterGPUBufferConsumer(this);
+	}
+}
+
 FString
 UDIFleshDeformer::GetDisplayName() const
 {
@@ -87,16 +96,25 @@ BEGIN_SHADER_PARAMETER_STRUCT(FFleshDeformerDataInterfaceParameters, )
 	SHADER_PARAMETER(uint32, BaseVertexIndex)
 	SHADER_PARAMETER(uint32, NumTetRestVertexBuffer)
 	SHADER_PARAMETER(uint32, NumTetVertexBuffer)
+
+	SHADER_PARAMETER(uint32, NumVertsToParentsBuffer)
+	SHADER_PARAMETER(int32, VertsToParentsBufferOffset)
+	SHADER_PARAMETER(uint32, VertsToParentsBufferStride)
+
 	SHADER_PARAMETER(uint32, NumParentsBuffer)
+	SHADER_PARAMETER(int32, ParentsBufferOffset)
+	SHADER_PARAMETER(uint32, ParentsBufferStride)
+
 	SHADER_PARAMETER(uint32, NumWeightsBuffer)
 	SHADER_PARAMETER(uint32, NumOffsetBuffer)
 	SHADER_PARAMETER(uint32, NumMaskBuffer)
-	SHADER_PARAMETER_RDG_BUFFER_SRV(Buffer<float>, TetRestVertexBuffer)
-	SHADER_PARAMETER_RDG_BUFFER_SRV(Buffer<float>, TetVertexBuffer)
-	SHADER_PARAMETER_RDG_BUFFER_SRV(Buffer<int>, ParentsBuffer)
-	SHADER_PARAMETER_RDG_BUFFER_SRV(Buffer<float>, WeightsBuffer)
-	SHADER_PARAMETER_RDG_BUFFER_SRV(Buffer<float>, OffsetBuffer)
-	SHADER_PARAMETER_RDG_BUFFER_SRV(Buffer<float>, MaskBuffer)
+	SHADER_PARAMETER_SRV(Buffer<float>, TetRestVertexBuffer)
+	SHADER_PARAMETER_SRV(Buffer<float>, TetVertexBuffer)
+	SHADER_PARAMETER_SRV(Buffer<uint32>, VertsToParentsBuffer)
+	SHADER_PARAMETER_SRV(Buffer<uint32>, ParentsBuffer)
+	SHADER_PARAMETER_SRV(Buffer<FFloat16>, WeightsBuffer)
+	SHADER_PARAMETER_SRV(Buffer<FFloat16>, OffsetBuffer)
+	SHADER_PARAMETER_SRV(Buffer<FFloat16>, MaskBuffer)
 END_SHADER_PARAMETER_STRUCT()
 
 void
@@ -152,7 +170,7 @@ UDIFleshDeformer::CreateDataProvider(TObjectPtr<UObject> InBinding, uint64 InInp
 		int32 NumFleshComponents = 0;
 		for (TSet<UActorComponent*>::TConstIterator It = Components.CreateConstIterator(); It; ++It)
 		{
-			NumFleshComponents += const_cast<UFleshComponent*>(Cast<UFleshComponent>(*It)) != nullptr;
+			NumFleshComponents += const_cast<UDeformableTetrahedralComponent*>(Cast<UDeformableTetrahedralComponent>(*It)) != nullptr;
 		}
 
 		// Find the flesh component that is associated with this skeletal mesh.
@@ -160,43 +178,66 @@ UDIFleshDeformer::CreateDataProvider(TObjectPtr<UObject> InBinding, uint64 InInp
 		{
 			if (const UActorComponent* Component = *It)
 			{
-				if (UFleshComponent* FleshComponent = const_cast<UFleshComponent*>(Cast<UFleshComponent>(Component)))
+				if (UDeformableTetrahedralComponent* FleshComponent = const_cast<UDeformableTetrahedralComponent*>(Cast<UDeformableTetrahedralComponent>(Component)))
 				{
 					if (const UFleshAsset* FleshAsset = FleshComponent->GetRestCollection())
 					{
-						//if ((USkinnedMeshComponent*)FleshComponent->TargetSkeletalMesh == Provider->SkinnedMesh)
+						if (FleshAsset->TargetDeformationSkeleton == Cast<USkeletalMesh>(Provider->SkinnedMesh->GetSkinnedAsset()))
+						{
+							Provider->FleshMesh = FleshComponent;
+							ProducerComponent = FleshComponent;
+							break;
+						}
+						else if (!FleshAsset->TargetDeformationSkeleton && NumFleshComponents == 1)
 						{
-							if (FleshAsset->TargetSkeletalMesh.Equals(Provider->SkinnedMesh->GetName()))
-							{
-								Provider->FleshMesh = FleshComponent;
-								break;
-							}
-							else if (FleshAsset->TargetSkeletalMesh.IsEmpty() && NumFleshComponents == 1)
-							{
-								// The typical case is probably that we only have 1 skeletal mesh and 1 flesh
-								// component.  So, in an effort to keep the easy case easy, if the TargetSkeletalMesh
-								// isn't set, just do the obvious thing and use it.
-								UE_LOG(LogFleshDeformer, Display,
-									TEXT("FleshAsset '%s' has no value for 'TargetSkeletalMesh', "
-										"using it for skeletal mesh '%s' anyway as it's the only one available."),
-									*FleshAsset->GetName(),
-									*Provider->SkinnedMesh->GetName());
-								Provider->FleshMesh = FleshComponent;
-								break;
-							}
-							else if (FleshAsset->TargetSkeletalMesh.IsEmpty())
-							{
-								UE_LOG(LogFleshDeformer, Warning,
-									TEXT("FleshAsset '%s' has no value for 'TargetSkeletalMesh', "
-										"not using it for skeletal mesh '%s' as there are other candidates."),
-									*FleshAsset->GetName(),
-									*Provider->SkinnedMesh->GetName());
-							}
+							// The typical case is probably that we only have 1 skeletal mesh and 1 flesh
+							// component.  So, in an effort to keep the easy case easy, if the TargetSkeletalMesh
+							// isn't set, just do the obvious thing and use it.
+							UE_LOG(LogFleshDeformer, Display,
+								TEXT("FleshAsset '%s.%s.RestCollection' has no value for 'TargetDeformationSkeleton', "
+									"using it for skeletal mesh '%s' anyway as it's the only UFleshComponent available."),
+								*Actor->GetName(),
+								*FleshAsset->GetName(),
+								*Provider->SkinnedMesh->GetName());
+							Provider->FleshMesh = FleshComponent;
+							ProducerComponent = FleshComponent;
+							break;
+						}
+						else if (!FleshAsset->TargetDeformationSkeleton)
+						{
+							UE_LOG(LogFleshDeformer, Warning,
+								TEXT("FleshAsset '%s.%s.RestCollection' has no value for 'TargetDeformationSkeleton'; "
+									"not using it for skeletal mesh '%s' as there are %d UFleshComponent candidates."),
+								*Actor->GetName(),
+								*FleshAsset->GetName(),
+								*Provider->SkinnedMesh->GetName(),
+								NumFleshComponents);
 						}
 					}
 				}
 			}
 		}
+
+		if (!Provider->FleshMesh)
+		{
+			if (NumFleshComponents)
+			{
+				UE_LOG(LogFleshDeformer, Error,
+					TEXT("Found %d UFleshComponents under actor '%s', none of which were associated with skeletal mesh '%s'. "
+						"Disambiguate the association by setting 'FleshAsset.RestCollection.TargetDeformationSkeleton' to the "
+						"skeletal mesh it should deform."),
+					NumFleshComponents,
+					*Actor->GetName(),
+					*Provider->SkinnedMesh->GetName());
+			}
+			else
+			{
+				UE_LOG(LogFleshDeformer, Error,
+					TEXT("No UFleshComponent found under actor '%s' for skeletal mesh '%s'."),
+					*Actor->GetName(),
+					*Provider->SkinnedMesh->GetName());
+			}
+		}
 	}
 	// TODO: Someday, if/when Optimus supports deforming static meshes...
 	//if (Provider->StaticMesh = Cast<UStaticMeshComponent>(InBinding))
@@ -228,7 +269,7 @@ UDIFleshDeformerDataProvider::GetRenderProxy()
 
 FDIFleshDeformerProviderProxy::FDIFleshDeformerProviderProxy(
 	USkinnedMeshComponent* SkinnedMeshComponentIn, 
-	UFleshComponent* FleshComponentIn,
+	UDeformableTetrahedralComponent* FleshComponentIn,
 	FFleshDeformerParameters const& FleshDeformerParametersIn)
 	: SkinnedMeshComponent(SkinnedMeshComponentIn)
 	, SkeletalMeshObject(SkinnedMeshComponentIn ? SkinnedMeshComponentIn->MeshObject : nullptr)
@@ -270,20 +311,15 @@ FDIFleshDeformerProviderProxy::IsValid(FValidationData const& InValidationData)
 		return false;
 	}
 
-	// Note: tet bindings will be invalid until initialized to a bindings group.
-	// We're just testing for the existence of such a group here.
-	GeometryCollection::Facades::FTetrahedralBindings TetBindings(*FleshCollection);
+	// This function is called every draw call.  It may be better to do this initialization
+	// higher in the stack, but we do it here to support on the fly LOD switching.  I think...
 	FName MeshName(*MeshId, MeshId.Len());
-	int32 TetMeshIdx = TetBindings.GetTetMeshIndex(MeshName, LodIndex);
-	if (TetMeshIdx == INDEX_NONE)
+	if (!FleshComponent->GetGPUBufferManager().InitGPUBindingsBuffer(SkinnedMeshComponent, MeshName, LodIndex))
 	{
-		UE_LOG(LogFleshDeformer, Error, 
-			TEXT("FleshDeformer: Failed to find tet mesh index associated to mesh '%s' LOD: %d"), 
-			*MeshId, LodIndex);
 		return false;
 	}
-	bool bValid = TetBindings.ContainsBindingsGroup(TetMeshIdx, MeshName, LodIndex);
-	return bValid;
+
+	return true;
 }
 
 struct FFleshDeformerDataInterfacePermutationIds
@@ -313,216 +349,7 @@ FDIFleshDeformerProviderProxy::GatherPermutations(FPermutationData& InOutPermuta
 
 void
 FDIFleshDeformerProviderProxy::AllocateResources(FRDGBuilder& GraphBuilder)
-{
-	//
-	// Rest Vertices from FleshComponent's rest collection.
-	//
-
-	check(FleshComponent);
-	const UFleshAsset* RestCollectionAsset = FleshComponent->GetRestCollection();
-	const FFleshCollection* FleshCollection = RestCollectionAsset->GetCollection();
-	NumTetRestVertices = 0;
-	NumTetVertices = 0;
-
-	const TManagedArray<FVector3f>* RestVertex = nullptr;
-	if (FleshCollection)
-	{
-		RestVertex = &FleshCollection->Vertex;
-		NumTetRestVertices = RestVertex->Num();
-	}
-
-	TetRestVertexBuffer = GraphBuilder.CreateBuffer(
-		FRDGBufferDesc::CreateStructuredDesc(sizeof(float), FGenericPlatformMath::Max((uint32)1, NumTetRestVertices * 3)),
-		TEXT("FleshDeformer.RestVertex"));
-	TetRestVertexBufferSRV = GraphBuilder.CreateSRV(TetRestVertexBuffer);
-
-	if (NumTetRestVertices)
-	{
-		const FVector3f* Vert0 = RestVertex->GetData();
-		const float* Vert00 = &Vert0->X;
-
-		GraphBuilder.QueueBufferUpload(
-			TetRestVertexBuffer,
-			Vert00,
-			NumTetRestVertices * 3 * sizeof(float),
-			ERDGInitialDataFlags::None);
-	}
-	else
-	{
-		GraphBuilder.QueueBufferUpload(
-			TetRestVertexBuffer,
-			&NullFloatBuffer,
-			1 * sizeof(float),
-			ERDGInitialDataFlags::None);
-	}
-
-	//
-	// Vertices from FleshComponent's dynamic collection, if they exist; rest
-	// collection vertices otherwise.
-	//
-
-	const TManagedArray<FVector3f>* Vertex = nullptr;
-	if (UFleshDynamicAsset* DynamicCollectionAsset = FleshComponent->GetDynamicCollection())
-	{
-		Vertex = &DynamicCollectionAsset->GetPositions();
-		NumTetVertices = Vertex->Num();
-	}
-	if (!NumTetVertices && FleshCollection)
-	{
-		Vertex = &FleshCollection->Vertex;
-		NumTetVertices = Vertex->Num();
-	}
-
-	TetVertexBuffer = GraphBuilder.CreateBuffer(
-		FRDGBufferDesc::CreateStructuredDesc(sizeof(float), FGenericPlatformMath::Max((uint32)1, NumTetVertices * 3)),
-		TEXT("FleshDeformer.Vertex"));
-	TetVertexBufferSRV = GraphBuilder.CreateSRV(TetVertexBuffer);
-
-	if (NumTetVertices)
-	{
-		const FVector3f* Vert0 = Vertex->GetData();
-		const float* Vert00 = &Vert0->X;
-
-		GraphBuilder.QueueBufferUpload(
-			TetVertexBuffer,
-			Vert00,
-			NumTetVertices * 3 * sizeof(float),
-			ERDGInitialDataFlags::None);
-	}
-	else
-	{
-		GraphBuilder.QueueBufferUpload(
-			TetVertexBuffer,
-			&NullFloatBuffer,
-			1 * sizeof(float),
-			ERDGInitialDataFlags::None);
-	}
-
-	//
-	// Bindings data from FleshComponent's rest collection.
-	//
-
-	GeometryCollection::Facades::FTetrahedralBindings* TetBindings = nullptr;
-	if (FleshCollection)
-	{
-		TetBindings = new GeometryCollection::Facades::FTetrahedralBindings(*FleshCollection);
-		FName MeshName(*MeshId, MeshId.Len());
-		TetIndex = TetBindings->GetTetMeshIndex(MeshName, LodIndex);
-		if (TetIndex != INDEX_NONE)
-		{
-			if (!TetBindings->ReadBindingsGroup(TetIndex, MeshName, LodIndex))
-			{
-				// Warn?
-			}
-		}
-		if (!TetBindings->IsValid())
-		{
-			delete TetBindings;
-			TetBindings = nullptr;
-		}
-	}
-	const TManagedArrayAccessor<FIntVector4>* Parents = TetBindings ? TetBindings->GetParentsRO() : nullptr;
-	const TManagedArrayAccessor<FVector4f>* Weights = TetBindings ? TetBindings->GetWeightsRO() : nullptr;
-	const TManagedArrayAccessor<FVector3f>* Offset = TetBindings ? TetBindings->GetOffsetsRO() : nullptr;
-	const TManagedArrayAccessor<float>* Mask = TetBindings ? TetBindings->GetMaskRO() : nullptr;
-
-	NumParents = (Parents ? Parents->Num() : 0);
-	ParentsBuffer =
-		GraphBuilder.CreateBuffer(
-			FRDGBufferDesc::CreateStructuredDesc(sizeof(int32), FGenericPlatformMath::Max((uint32)1, NumParents * 4)),
-			TEXT("FleshDeformer.Parents"));
-	ParentsBufferSRV = GraphBuilder.CreateSRV(ParentsBuffer);
-
-	NumWeights = (Weights ? Weights->Num() : 0);
-	WeightsBuffer =
-		GraphBuilder.CreateBuffer(
-			FRDGBufferDesc::CreateStructuredDesc(sizeof(float), FGenericPlatformMath::Max((uint32)1, NumWeights * 4)),
-			TEXT("FleshDeformer.Weights"));
-	WeightsBufferSRV = GraphBuilder.CreateSRV(WeightsBuffer);
-
-	NumOffset = (Offset ? Offset->Num() : 0);
-	OffsetBuffer =
-		GraphBuilder.CreateBuffer(
-			FRDGBufferDesc::CreateStructuredDesc(sizeof(float), FGenericPlatformMath::Max((uint32)1, NumOffset * 3)),
-			TEXT("FleshDeformer.Offset"));
-	OffsetBufferSRV = GraphBuilder.CreateSRV(OffsetBuffer);
-
-	NumMask = (Mask ? Mask->Num() : 0);
-	MaskBuffer =
-		GraphBuilder.CreateBuffer(
-			FRDGBufferDesc::CreateStructuredDesc(sizeof(float), FGenericPlatformMath::Max((uint32)1, NumMask)),
-			TEXT("FleshDeformer.Mask"));
-	MaskBufferSRV = GraphBuilder.CreateSRV(MaskBuffer);
-
-	if (Parents && Parents->Num())
-	{
-		const TManagedArray<FIntVector4>& ParentsValues = Parents->Get();
-		GraphBuilder.QueueBufferUpload(
-			ParentsBuffer,
-			&ParentsValues[0][0],
-			ParentsValues.Num() * 4 * sizeof(int32),
-			ERDGInitialDataFlags::None);
-	}
-	else
-	{
-		GraphBuilder.QueueBufferUpload(
-			ParentsBuffer,
-			&NullIntBuffer,
-			1 * sizeof(int32),
-			ERDGInitialDataFlags::None);
-	}
-	if (Weights && Weights->Num())
-	{
-		const TManagedArray<FVector4f>& WeightsValues = Weights->Get();
-		GraphBuilder.QueueBufferUpload(
-			WeightsBuffer,
-			&WeightsValues.GetData()->X,
-			WeightsValues.Num() * 4 * sizeof(float),
-			ERDGInitialDataFlags::None);
-	}
-	else
-	{
-		GraphBuilder.QueueBufferUpload(
-			WeightsBuffer,
-			&NullFloatBuffer,
-			1 * sizeof(float),
-			ERDGInitialDataFlags::None);
-	}
-	if (Offset && Offset->Num())
-	{
-		const TManagedArray<FVector3f>& OffsetValues = Offset->Get();
-		GraphBuilder.QueueBufferUpload(
-			OffsetBuffer,
-			&OffsetValues.GetData()->X,
-			OffsetValues.Num() * 3 * sizeof(float),
-			ERDGInitialDataFlags::None);
-	}
-	else
-	{
-		GraphBuilder.QueueBufferUpload(
-			OffsetBuffer,
-			&NullFloatBuffer,
-			1 * sizeof(float),
-			ERDGInitialDataFlags::None);
-	}
-	if (Mask && Mask->Num())
-	{
-		const TManagedArray<float>& MaskValues = Mask->Get();
-		GraphBuilder.QueueBufferUpload(
-			MaskBuffer,
-			&MaskValues[0],
-			MaskValues.Num() * sizeof(float),
-			ERDGInitialDataFlags::None);
-	}
-	else
-	{
-		GraphBuilder.QueueBufferUpload(
-			MaskBuffer,
-			&NullFloatBuffer,
-			1 * sizeof(float),
-			ERDGInitialDataFlags::None);
-	}
-}
+{}
 
 void
 FDIFleshDeformerProviderProxy::GatherDispatchData(FDispatchData const& InDispatchData)
@@ -531,6 +358,14 @@ FDIFleshDeformerProviderProxy::GatherDispatchData(FDispatchData const& InDispatc
 	FSkeletalMeshLODRenderData const* LodRenderData = &SkeletalMeshRenderData.LODRenderData[LodIndex];
 	const TStridedView<FParameters> ParameterArray = MakeStridedParameterView<FParameters>(InDispatchData);
 
+	FName MeshName(*MeshId, MeshId.Len());
+	Chaos::Softs::FChaosFleshDeformableGPUManager& GPUManager = FleshComponent->GetGPUBufferManager();
+	Chaos::Softs::FChaosFleshDeformableGPUManager::FBindingsBuffer* Bindings = GPUManager.GetBindingsBuffer(SkinnedMeshComponent, MeshName, LodIndex);
+	if (!Bindings)
+	{
+		return;
+	}
+
 	for (int32 InvocationIndex = 0; InvocationIndex < ParameterArray.Num(); ++InvocationIndex)
 	{
 		FSkelMeshRenderSection const& RenderSection = LodRenderData->RenderSections[InvocationIndex];
@@ -540,20 +375,62 @@ FDIFleshDeformerProviderProxy::GatherDispatchData(FDispatchData const& InDispatc
 		Parameters.NumVertices = RenderSection.GetNumVertices();
 		Parameters.BaseVertexIndex = RenderSection.BaseVertexIndex;
 
-		Parameters.NumTetRestVertexBuffer = NumTetRestVertices;
-		Parameters.NumTetVertexBuffer = NumTetVertices;
-
-		Parameters.NumParentsBuffer = NumParents;
-		Parameters.NumWeightsBuffer = NumWeights;
-		Parameters.NumOffsetBuffer = NumOffset;
-		Parameters.NumMaskBuffer = NumMask;
-
-		Parameters.TetRestVertexBuffer = TetRestVertexBufferSRV;
-		Parameters.TetVertexBuffer = TetVertexBufferSRV;
-
-		Parameters.ParentsBuffer = ParentsBufferSRV;
-		Parameters.WeightsBuffer = WeightsBufferSRV;
-		Parameters.OffsetBuffer = OffsetBufferSRV;
-		Parameters.MaskBuffer = MaskBufferSRV;
+		//
+		// Time varying
+		//
+
+		// Sparsified tet vertices - only referenced parent vertices.
+		// SHADER_PARAMETER(uint32, NumTetVertexBuffer)
+		// SHADER_PARAMETER_SRV(Buffer<float>, TetVertexBuffer)
+		Parameters.NumTetVertexBuffer = Bindings->GetVerticesBuffer().GetNumValues();
+		Parameters.TetVertexBuffer = Bindings->GetVerticesBuffer().ShaderResourceViewRHI;
+
+		//
+		// Unvarying (typically)
+		//
+
+		// Sparsified rest tet vertices - only referenced parent vertices.
+		// SHADER_PARAMETER(uint32, NumTetRestVertexBuffer)
+		// SHADER_PARAMETER_SRV(Buffer<float>, TetRestVertexBuffer)
+		Parameters.NumTetRestVertexBuffer = Bindings->GetRestVerticesBuffer().GetNumValues();
+		Parameters.TetRestVertexBuffer = Bindings->GetRestVerticesBuffer().ShaderResourceViewRHI;
+
+		// Per vertex packed (8, 16, or 32 bit) index to sparsified parents
+		// SHADER_PARAMETER(uint32, NumVertsToParentsBuffer)
+		// SHADER_PARAMETER(int32, VertsToParentsBufferOffset)
+		// SHADER_PARAMETER(uint32, VertsToParentsBufferStride)
+		// SHADER_PARAMETER_SRV(Buffer<uint32>, VertsToParentsBuffer)
+		Parameters.NumVertsToParentsBuffer = Bindings->GetVertsToParentsBuffer().GetNumValues();
+		Parameters.VertsToParentsBufferOffset = Bindings->GetVertsToParentsBuffer().GetOffset();
+		Parameters.VertsToParentsBufferStride = Bindings->GetVertsToParentsBuffer().GetDataStride();
+		Parameters.VertsToParentsBuffer = Bindings->GetVertsToParentsBuffer().ShaderResourceViewRHI;
+
+		// Sparsified (only unique) and packed (8, 16, or 32 bit) parents.
+		// SHADER_PARAMETER(uint32, NumParentsBuffer)
+		// SHADER_PARAMETER(int32, ParentsBufferOffset)
+		// SHADER_PARAMETER(uint32, ParentsBufferStride)
+		// SHADER_PARAMETER_SRV(Buffer<uint32>, ParentsBuffer)
+		Parameters.NumParentsBuffer = Bindings->GetParentsBuffer().GetNumValues();
+		Parameters.ParentsBufferOffset = Bindings->GetParentsBuffer().GetOffset();
+		Parameters.ParentsBufferStride = Bindings->GetParentsBuffer().GetDataStride();
+		Parameters.ParentsBuffer = Bindings->GetParentsBuffer().ShaderResourceViewRHI;
+
+		// Per vertex half precision floating point barycentric weights.
+		// SHADER_PARAMETER(uint32, NumWeightsBuffer)
+		//SHADER_PARAMETER_SRV(Buffer<FFloat16>, WeightsBuffer)
+		Parameters.NumWeightsBuffer = Bindings->GetWeightsBuffer().GetNumValues();
+		Parameters.WeightsBuffer = Bindings->GetWeightsBuffer().ShaderResourceViewRHI;
+
+		// Per vertex half precision floating point offset vector.
+		//SHADER_PARAMETER(uint32, NumOffsetBuffer)
+		//SHADER_PARAMETER_SRV(Buffer<FFloat16>, OffsetBuffer)
+		Parameters.NumOffsetBuffer = Bindings->GetOffsetsBuffer().GetNumValues();
+		Parameters.OffsetBuffer = Bindings->GetOffsetsBuffer().ShaderResourceViewRHI;
+
+		// Per vertex half precision floating point mask (amount).
+		//SHADER_PARAMETER(uint32, NumMaskBuffer)
+		//SHADER_PARAMETER_SRV(Buffer<FFloat16>, MaskBuffer)
+		Parameters.NumMaskBuffer = Bindings->GetOffsetsBuffer().GetNumValues();
+		Parameters.MaskBuffer = Bindings->GetMaskBuffer().ShaderResourceViewRHI;
 	}
 }
diff --git a/Engine/Plugins/Experimental/ChaosFlesh/Source/ChaosFleshEngine/Public/ChaosFlesh/ChaosDeformableGPUBuffers.h b/Engine/Plugins/Experimental/ChaosFlesh/Source/ChaosFleshEngine/Public/ChaosFlesh/ChaosDeformableGPUBuffers.h
new file mode 100644
--- /dev/null
+++ b/Engine/Plugins/Experimental/ChaosFlesh/Source/ChaosFleshEngine/Public/ChaosFlesh/ChaosDeformableGPUBuffers.h
@@ -0,0 +1,104 @@
+// Copyright Epic Games, Inc. All Rights Reserved.
+#pragma once
+
+#include "RenderResource.h"
+#include "RawIndexBuffer.h"
+
+namespace UE::ChaosDeformable
+{
+	/**
+	* Floating point GPU buffer.
+	*/
+	class FFloatArrayBufferWithSRV : public FVertexBufferWithSRV
+	{
+	public:
+		void SetBufferName(const FString& InBufferName) { BufferName = InBufferName; }
+
+		void Init(const float* InArray, const int32 Num);
+		void Init(const TArray<float>& InArray) { Array = InArray; NumValues = Array.Num(); }
+		void Init(TArray<float>&& InArray) { Array = MoveTemp(InArray); NumValues = Array.Num(); }
+
+		int32 GetNumValues() const { return NumValues; }
+		uint32 GetBufferSize() const;
+
+		/** Retain local memory buffer for reuse, rather than reallocating. */
+		void RetainLocalMemory() { bEmptyArray = false; }
+
+		/** FRenderResource overrides */
+		void InitRHI() override;
+		FString GetFriendlyName() const override { return BufferName; }
+		/** End FRenderResource overrides */
+	private:
+		bool bEmptyArray = true;
+		int32 NumValues = INDEX_NONE;
+		FString BufferName;
+		TArray<float> Array;
+	};
+
+	/**
+	* Floating point GPU buffer converting float to Half/FFloat16.
+	*/
+	class FHalfArrayBufferWithSRV : public FVertexBufferWithSRV
+	{
+	public:
+		void SetBufferName(const FString& InBufferName) { BufferName = InBufferName; }
+
+		void Init(const FFloat16* InArray, const int32 Num);
+		void Init(const TArray<FFloat16>& InArray) { Init(InArray.GetData(), InArray.Num()); }
+		void Init(const float* InArray, const int32 Num);
+		void Init(const TArray<float>& InArray);
+		void Init(const FVector3f* InArray, const int32 Num) { Init(Num ? &InArray[0].X : nullptr, Num * 3); }
+		void Init(const FVector4f* InArray, const int32 Num) { Init(Num ? &InArray[0].X : nullptr, Num * 4); }
+
+		int32 GetNumValues() const { return NumValues; }
+		uint32 GetBufferSize() const;
+
+		/** FRenderResource overrides */
+		void InitRHI() override;
+		FString GetFriendlyName() const override { return BufferName; }
+		/** End FRenderResource overrides */
+	private:
+		int32 NumValues = INDEX_NONE;
+		FString BufferName;
+		TArray<FFloat16> Array;
+	};
+
+	/**
+	* Integral GPU buffer compacting int32 to the smallest possible of uint8, uint16, or uint32.
+	*/
+	class FIndexArrayBufferWithSRV : public FVertexBufferWithSRV
+	{
+	public:
+		void SetBufferName(const FString& InBufferName) { BufferName = InBufferName; }
+
+		//! Convert float to uint8 with clamping, and optionaly shifting via \p MultOffset.
+		void Init(const int32* InArray, const int32 Num);
+		void Init(const uint32* InArray, const int32 Num);
+		void Init(const TArray<int32>& InArray) { Init(InArray.Num() ? &InArray[0] : nullptr, InArray.Num()); }
+		void Init(const FIntVector2* InArray, const int32 Num) { Init(Num ? &InArray[0].X : nullptr, Num * 2); }
+		void Init(const FIntVector3* InArray, const int32 Num) { Init(Num ? &InArray[0].X : nullptr, Num * 3); }
+		void Init(const FIntVector4* InArray, const int32 Num) { Init(Num ? &InArray[0].X : nullptr, Num * 4); }
+		void Init(const TArray<FIntVector4>& InArray) { Init(InArray.Num() ? &InArray[0].X : nullptr, InArray.Num() * 4); }
+
+		void Force32BitPacking() { bForce32 = true; }
+
+		int32 GetNumValues() const { return NumValues; }
+		int32 GetDataStride() const { return bForce32 ? sizeof(uint32) : bUint8 ? sizeof(uint8) : bUint16 ? sizeof(uint16) : sizeof(uint32); }
+		int32 GetOffset() const { return Offset; }
+		uint32 GetBufferSize() const;
+
+		/** FRenderResource overrides */
+		void InitRHI() override;
+		FString GetFriendlyName() const override { return BufferName; }
+		/** End FRenderResource overrides */
+	private:
+		bool bUint8 = false;
+		bool bUint16 = false;
+		bool bForce32 = false;
+		int32 Offset = 0;
+		int32 NumValues = INDEX_NONE;
+		FString BufferName;
+		TArray<uint32> Array;
+	};
+
+} // namespace UE::ChaosDeformable
diff --git a/Engine/Plugins/Experimental/ChaosFlesh/Source/ChaosFleshEngine/Public/ChaosFlesh/ChaosDeformablePhysicsComponent.h b/Engine/Plugins/Experimental/ChaosFlesh/Source/ChaosFleshEngine/Public/ChaosFlesh/ChaosDeformablePhysicsComponent.h
--- a/Engine/Plugins/Experimental/ChaosFlesh/Source/ChaosFleshEngine/Public/ChaosFlesh/ChaosDeformablePhysicsComponent.h
+++ b/Engine/Plugins/Experimental/ChaosFlesh/Source/ChaosFleshEngine/Public/ChaosFlesh/ChaosDeformablePhysicsComponent.h
@@ -40,7 +40,7 @@ public:
 	virtual void RemoveProxy(Chaos::Softs::FDeformableSolver::FGameThreadAccess& GameThreadSolver);
 	virtual void PreSolverUpdate() {};
 	virtual FDataMapValue NewDeformableData() { return FDataMapValue(nullptr); }
-	virtual void UpdateFromSimualtion(const FDataMapValue* SimualtionBuffer) {}
+	virtual void UpdateFromSimulation(const FDataMapValue* SimulationBuffer) {}
 
 
 	virtual void OnCreatePhysicsState() override;
diff --git a/Engine/Plugins/Experimental/ChaosFlesh/Source/ChaosFleshEngine/Public/ChaosFlesh/ChaosDeformableTetrahedralComponent.h b/Engine/Plugins/Experimental/ChaosFlesh/Source/ChaosFleshEngine/Public/ChaosFlesh/ChaosDeformableTetrahedralComponent.h
--- a/Engine/Plugins/Experimental/ChaosFlesh/Source/ChaosFleshEngine/Public/ChaosFlesh/ChaosDeformableTetrahedralComponent.h
+++ b/Engine/Plugins/Experimental/ChaosFlesh/Source/ChaosFleshEngine/Public/ChaosFlesh/ChaosDeformableTetrahedralComponent.h
@@ -3,10 +3,11 @@
 
 #include "Chaos/Deformable/ChaosDeformableSolverProxy.h"
 #include "Chaos/Deformable/ChaosDeformableSolver.h"
-#include "ChaosFlesh/ChaosDeformableSolverThreading.h"
 #include "ChaosFlesh/ChaosDeformablePhysicsComponent.h"
+#include "ChaosFlesh/ChaosDeformableSolverThreading.h"
 #include "ChaosFlesh/FleshAsset.h"
 #include "ChaosFlesh/FleshDynamicAsset.h"
+#include "ChaosFlesh/ChaosFleshDeformerBufferManager.h"
 #include "ChaosFlesh/SimulationAsset.h"
 #include "Components/MeshComponent.h"
 #include "UObject/ObjectMacros.h"
@@ -16,6 +17,7 @@
 class FFleshCollection;
 class ADeformableSolverActor;
 class UDeformableSolverComponent;
+class FChaosDeformableTetrahedralSceneProxy;
 
 /**
 *  Options for binding positions query.
@@ -37,23 +39,21 @@ struct FFleshSimulationSpaceGroup
 {
 	GENERATED_USTRUCT_BODY()
 
-		/**
-		* Bone from the associated skeletal mesh (indicated by RestCollection.TargetSkeletalMesh) to use as
-		* the space the sim runs in.
-		*/
+	/**
+	* Bone from the associated skeletal mesh (indicated by RestCollection.TargetSkeletalMesh) to use as
+	* the space the sim runs in.
+	*/
 	UPROPERTY(EditAnywhere, Category = "Physics", meta = (GetOptions = "GetSimSpaceBoneNameOptions", EditCondition = "SimSpace == ChaosDeformableSimSpace::Bone"))
 	FName SimSpaceBoneName;
 
-
 	/** Space the simulation will run in. */
 	UPROPERTY(EditAnywhere, Category = "Physics")
-		TEnumAsByte<ChaosDeformableSimSpace> SimSpace = ChaosDeformableSimSpace::World;
+	TEnumAsByte<ChaosDeformableSimSpace> SimSpace = ChaosDeformableSimSpace::World;
 
 	/** The skeletal mesh to use pull the \c SimSpaceBoneName from. */
 	UPROPERTY()
 	TObjectPtr<USkeletalMesh> SimSpaceSkeletalMesh;
 
-
 	int32 SimSpaceTransformIndex = INDEX_NONE;
 	int32 SimSpaceTransformGlobalIndex = INDEX_NONE;
 };
@@ -68,16 +68,16 @@ struct FBodyForcesGroup
 	bool bApplyGravity = true;
 
 	UPROPERTY(EditAnywhere, BlueprintReadWrite, Category = "Physics", meta = (ClampMin = "0.0", ClampMax = "1.0", UIMin = "0.0", UIMax = "1.0"))
-		float DampingMultiplier = 1.f;
+	float DampingMultiplier = 1.f;
 
 	UPROPERTY(EditAnywhere, BlueprintReadWrite, Category = "Physics", meta = (ClampMin = "0.0", ClampMax = "1.0", UIMin = "0.0", UIMax = "1.0"))
-		float StiffnessMultiplier = 1.f;
+	float StiffnessMultiplier = 1.f;
 
 	UPROPERTY(EditAnywhere, BlueprintReadWrite, Category = "Physics", meta = (ClampMin = "0.0", ClampMax = "1.0", UIMin = "0.0", UIMax = "1.0"))
-		float IncompressibilityMultiplier = 1.f;
+	float IncompressibilityMultiplier = 1.f;
 
 	UPROPERTY(EditAnywhere, BlueprintReadWrite, Category = "Physics", meta = (ClampMin = "0.0", ClampMax = "1.0", UIMin = "0.0", UIMax = "1.0"))
-		float InflationMultiplier = 1.f;
+	float InflationMultiplier = 1.f;
 };
 
 /**
@@ -102,10 +102,16 @@ public:
 	virtual void EndPlay(const EEndPlayReason::Type ReasonEnd) override;
 	void Invalidate();
 
-	/** Simulation Interface*/
+	/** UPrimitiveComponent Interface */
+	//virtual FPrimitiveSceneProxy* CreateSceneProxy() override;
+
+	/** GPU Interface */
+	Chaos::Softs::FChaosFleshDeformableGPUManager& GetGPUBufferManager() { return GPUBufferManager; }
+
+	/** Simulation Interface */
 	virtual FThreadingProxy* NewProxy() override;
 	virtual FDataMapValue NewDeformableData() override;
-	virtual void UpdateFromSimualtion(const FDataMapValue* SimualtionBuffer) override;
+	virtual void UpdateFromSimulation(const FDataMapValue* SimualtionBuffer) override;
 
 	/** RestCollection */
 	void SetRestCollection(const UFleshAsset * InRestCollection);
@@ -147,6 +153,9 @@ public:
 	UPROPERTY(EditAnywhere, Category = "Rendering")
 	TObjectPtr<UProceduralMeshComponent> Mesh;
 
+	UPROPERTY(EditAnywhere, Category = "Rendering")
+	TArray<int32> HideTetrahedra;
+
 private:
 	/** FleshAsset that describes the simulation rest state. */
 	UPROPERTY(EditAnywhere, Category = "Physics", meta = (DisplayPriority = 1))
@@ -160,7 +169,6 @@ private:
 	UPROPERTY()
 	TObjectPtr<USimulationAsset> SimulationCollection;
 
-
 	/* Returns a list of bone names from the currently selected skeletal mesh. */
 	UFUNCTION(CallInEditor)
 	TArray<FString> GetSimSpaceBoneNameOptions() const;
@@ -200,5 +208,7 @@ private:
 	TArray<FVector> GetSkeletalMeshBindingPositionsInternal(const USkeletalMesh* InSkeletalMesh, TArray<bool>* OutInfluence = nullptr) const;
 	void DebugDrawSkeletalMeshBindingPositions() const;
 
+	//FChaosDeformableTetrahedralSceneProxy* RenderProxy = nullptr;
+	Chaos::Softs::FChaosFleshDeformableGPUManager GPUBufferManager;
 };
 
diff --git a/Engine/Plugins/Experimental/ChaosFlesh/Source/ChaosFleshEngine/Public/ChaosFlesh/ChaosDeformableTypes.h b/Engine/Plugins/Experimental/ChaosFlesh/Source/ChaosFleshEngine/Public/ChaosFlesh/ChaosDeformableTypes.h
--- a/Engine/Plugins/Experimental/ChaosFlesh/Source/ChaosFleshEngine/Public/ChaosFlesh/ChaosDeformableTypes.h
+++ b/Engine/Plugins/Experimental/ChaosFlesh/Source/ChaosFleshEngine/Public/ChaosFlesh/ChaosDeformableTypes.h
@@ -33,6 +33,11 @@ struct FChaosEngineDeformableCVarParams
 
 	int32 EnvCollisionsLineTraceBatchSize = 10;
 
+	bool bTestUnpacking = false;
+	bool bTestSparseMappings = false;
+	bool bTestCompressedIndexing = false;
+	bool bUpdateGPUBuffersOnTick = true;
+
 	bool IsDebugDrawingEnabled()
 	{
 #if WITH_EDITOR
diff --git a/Engine/Plugins/Experimental/ChaosFlesh/Source/ChaosFleshEngine/Public/ChaosFlesh/ChaosFleshDeformerBufferManager.h b/Engine/Plugins/Experimental/ChaosFlesh/Source/ChaosFleshEngine/Public/ChaosFlesh/ChaosFleshDeformerBufferManager.h
new file mode 100644
--- /dev/null
+++ b/Engine/Plugins/Experimental/ChaosFlesh/Source/ChaosFleshEngine/Public/ChaosFlesh/ChaosFleshDeformerBufferManager.h
@@ -0,0 +1,144 @@
+// Copyright Epic Games, Inc. All Rights Reserved.
+
+#pragma once
+#include "ChaosFlesh/ChaosDeformableGPUBuffers.h"
+#include "GeometryCollection/ManagedArrayAccessor.h"
+#include "RenderCommandFence.h"
+
+class UDeformableTetrahedralComponent;
+class FFleshCollection;
+
+namespace Chaos::Softs 
+{
+	class /*CHAOS_API*/ FChaosFleshDeformableGPUManager
+	{
+	public:
+		class FBindingsBuffer
+		{
+		public:
+			FBindingsBuffer();
+			~FBindingsBuffer();
+
+			void SetOwnerName(const FName& OwnerName);
+
+			/** Initialize GPU buffers. */
+			bool InitData(const TManagedArrayAccessor<FIntVector4>* ParentsArray, const TManagedArrayAccessor<FVector4f>* WeightsArray, const TManagedArrayAccessor<FVector3f>* OffsetArray, const TManagedArrayAccessor<float>* MaskArray, const TManagedArray<FVector3f>* RestVerticesArray, const TManagedArray<FVector3f>* VerticesArray);
+
+			/**
+			* Update local sparse vertices buffer from dense vertices, and uploads to GPU.
+			* Returns \c true if \p DenseVertex differs from prior values.
+			*/
+			bool UpdateVertices(const TArray<FVector3f>& DenseVertex);
+
+			UE::ChaosDeformable::FFloatArrayBufferWithSRV& GetRestVerticesBuffer() { return RestVerticesBuffer; }
+			const UE::ChaosDeformable::FFloatArrayBufferWithSRV& GetRestVerticesBuffer() const { return RestVerticesBuffer; }
+
+			UE::ChaosDeformable::FFloatArrayBufferWithSRV& GetVerticesBuffer() { return VerticesBuffer; }
+			const UE::ChaosDeformable::FFloatArrayBufferWithSRV& GetVerticesBuffer() const { return VerticesBuffer; }
+
+			UE::ChaosDeformable::FIndexArrayBufferWithSRV& GetParentsBuffer() { return ParentsBuffer; }
+			const UE::ChaosDeformable::FIndexArrayBufferWithSRV& GetParentsBuffer() const { return ParentsBuffer; }
+
+			UE::ChaosDeformable::FIndexArrayBufferWithSRV& GetVertsToParentsBuffer() { return VertsToParentsBuffer; }
+			const UE::ChaosDeformable::FIndexArrayBufferWithSRV& GetVertsToParentsBuffer() const { return VertsToParentsBuffer; }
+
+			UE::ChaosDeformable::FHalfArrayBufferWithSRV& GetWeightsBuffer() { return WeightsBuffer; }
+			const UE::ChaosDeformable::FHalfArrayBufferWithSRV& GetWeightsBuffer() const { return WeightsBuffer; }
+
+			UE::ChaosDeformable::FHalfArrayBufferWithSRV& GetOffsetsBuffer() { return OffsetsBuffer; }
+			const UE::ChaosDeformable::FHalfArrayBufferWithSRV& GetOffsetsBuffer() const { return OffsetsBuffer; }
+
+			UE::ChaosDeformable::FHalfArrayBufferWithSRV& GetMaskBuffer() { return MaskBuffer; }
+			const UE::ChaosDeformable::FHalfArrayBufferWithSRV& GetMaskBuffer() const { return MaskBuffer; }
+
+			//const TArray<int32>& GetCondensationMapping() const { return CondensationMapping; }
+			//const TArray<FIntVector4>& GetSparseParents() const { return SparseParents; }
+			//const TArray<FVector3f>& GetSparseRestVertices() const { return SparseRestVertices; }
+			//const TArray<FVector3f>& GetSparseVertices() const { return SparseVertices; }
+
+		private:
+			bool UpdateSparseVertices(const TArray<FVector3f>& DenseVertexIn, TArray<FVector3f>& SparseVerticesIn) const;
+
+			/**
+			* Create a mapping that reduces the number of vertices referenced by the parents array,
+			* and the vertices required by culling unreferenced vertices.
+			*/
+			void InitCondensationMapping(const TManagedArrayAccessor<FIntVector4>& ParentsArray);
+
+			/** Release GPU resources. */
+			void DestroyGPUBuffers();
+
+			/** Fence used in render thread cleanup on destruction. */
+			FRenderCommandFence RenderResourceDestroyFence;
+
+			// Non-const
+			UE::ChaosDeformable::FFloatArrayBufferWithSRV VerticesBuffer;
+
+			// Const
+			UE::ChaosDeformable::FFloatArrayBufferWithSRV RestVerticesBuffer;
+			UE::ChaosDeformable::FIndexArrayBufferWithSRV ParentsBuffer;
+			UE::ChaosDeformable::FIndexArrayBufferWithSRV VertsToParentsBuffer;
+			UE::ChaosDeformable::FHalfArrayBufferWithSRV WeightsBuffer;
+			UE::ChaosDeformable::FHalfArrayBufferWithSRV OffsetsBuffer;
+			UE::ChaosDeformable::FHalfArrayBufferWithSRV MaskBuffer;
+
+			// Sparse vertex index to dense (excludes -1) - maps dense vertices to SparseVertices.
+			TArray<int32> SparseToDenseIndices;
+			// For each render vertex, the corresponding sparse parent index.
+			TArray<uint32> VertsToSparseParents;
+			// Sparse (unique) parents.
+			TArray<FIntVector4> SparseParents;
+
+			// Local buffer for sparse rest vertices.
+			TArray<FVector3f> SparseRestVertices;
+			// Local buffer for sparse time varying vertices.
+			TArray<FVector3f> SparseVertices;
+		};
+
+	public:
+		FChaosFleshDeformableGPUManager();
+		FChaosFleshDeformableGPUManager(UDeformableTetrahedralComponent* InComponent);
+		virtual ~FChaosFleshDeformableGPUManager();
+
+		void SetOwner(UDeformableTetrahedralComponent* InComponent);
+
+		/** Register a consumer of a buffer this class will manage. */
+		void RegisterGPUBufferConsumer(const void* ID) const { GPUBufferConsumers.AddUnique(ID); }
+		/** Remove a registered consumer. */
+		void UnRegisterGPUBufferConsumer(const void* ID) const;
+		/** Returns \c true if there's at least 1 registered consumer. */
+		bool HasRegisteredGPUBufferConsumer() const { return !GPUBufferConsumers.IsEmpty(); }
+
+		/** Initialize the binding data for a render mesh. Will register \p ID as consumer. Returns success. */
+		bool InitGPUBindingsBuffer(const void* ID, const FName MeshName, const int32 LodIndex, const bool ForceInit=false);
+		/** Returns true if we have a matching bindings buffer. */
+		bool HasGPUBindingsBuffer(const void* ID, const FName MeshName, const int32 LodIndex) const;
+
+		/** 
+		 * If a consumer has been registered, updates time varying data from the component's 
+		 * dynamic collection. Returns success. 
+		 */
+		bool UpdateGPUBuffers();
+
+		/** Look up bindings data buffers associated with \p ID, \p MeshName, and \p LodIndex. */
+		const FBindingsBuffer* GetBindingsBuffer(const void* ID, const FName MeshName, const int32 LodIndex) const;
+		FBindingsBuffer* GetBindingsBuffer(const void* ID, const FName MeshName, const int32 LodIndex);
+
+	private:
+		UDeformableTetrahedralComponent* Component = nullptr;
+
+		// Convenience functions
+		const FFleshCollection* GetRestCollection() const;
+		const TManagedArray<FVector3f>* GetRestVertexArray() const;
+		const TManagedArray<FVector3f>* GetVertexArray() const;
+
+		// Consumers
+		mutable TArray<const void*> GPUBufferConsumers;
+
+		// Data tables
+		typedef TTuple<void*, FName, int32> BufferId;
+		typedef TMap<BufferId, TSharedPtr<FBindingsBuffer>> IdToBuffers;
+		IdToBuffers BufferTable;
+	};
+
+} // namespace Chaos::Softs
\ No newline at end of file
diff --git a/Engine/Plugins/Experimental/ChaosFlesh/Source/ChaosFleshEngine/Public/ChaosFlesh/FleshAsset.h b/Engine/Plugins/Experimental/ChaosFlesh/Source/ChaosFleshEngine/Public/ChaosFlesh/FleshAsset.h
--- a/Engine/Plugins/Experimental/ChaosFlesh/Source/ChaosFleshEngine/Public/ChaosFlesh/FleshAsset.h
+++ b/Engine/Plugins/Experimental/ChaosFlesh/Source/ChaosFleshEngine/Public/ChaosFlesh/FleshAsset.h
@@ -117,7 +117,7 @@ public:
 	TObjectPtr<USkeleton> Skeleton;
 
 	/**
-	* Skeleton to use with \c GetSkeletalMeshEmbeddedPositions() on the flesh component. 
+	* Skeleton to use with the flesh deformer or \c GetSkeletalMeshEmbeddedPositions() on the flesh component. 
 	* Bindings for this skeletal mesh must be stored in the rest collection.
 	*/
 	UPROPERTY(EditAnywhere, BlueprintReadWrite, Category = "Rendering")
@@ -129,12 +129,6 @@ public:
 	UPROPERTY(EditAnywhere, Category = "Geometry")
 	TObjectPtr<UStaticMesh> StaticMesh;
 
-	//
-	// Deformer Target Mesh
-	//
-	UPROPERTY(EditAnywhere, Category = "Deformer")
-	FString TargetSkeletalMesh;
-
 #if WITH_EDITORONLY_DATA
 	UPROPERTY(EditAnywhere, Category = "Render")
 	bool bRenderInEditor = true;
diff --git a/Engine/Plugins/Experimental/ChaosFlesh/Source/ChaosFleshEngine/Public/ChaosFlesh/FleshComponent.h b/Engine/Plugins/Experimental/ChaosFlesh/Source/ChaosFleshEngine/Public/ChaosFlesh/FleshComponent.h
--- a/Engine/Plugins/Experimental/ChaosFlesh/Source/ChaosFleshEngine/Public/ChaosFlesh/FleshComponent.h
+++ b/Engine/Plugins/Experimental/ChaosFlesh/Source/ChaosFleshEngine/Public/ChaosFlesh/FleshComponent.h
@@ -31,6 +31,4 @@ class CHAOSFLESHENGINE_API UFleshComponent : public UDeformableGameplayComponent
 	GENERATED_UCLASS_BODY()
 public:
 
-	TArray<int32> HideTetrahedra;
-
-};
\ No newline at end of file
+};
diff --git a/Engine/Plugins/Experimental/ChaosFlesh/Source/ChaosFleshEngine/Public/DataInterfaces/DIFleshDeformer.h b/Engine/Plugins/Experimental/ChaosFlesh/Source/ChaosFleshEngine/Public/DataInterfaces/DIFleshDeformer.h
--- a/Engine/Plugins/Experimental/ChaosFlesh/Source/ChaosFleshEngine/Public/DataInterfaces/DIFleshDeformer.h
+++ b/Engine/Plugins/Experimental/ChaosFlesh/Source/ChaosFleshEngine/Public/DataInterfaces/DIFleshDeformer.h
@@ -3,7 +3,7 @@
 
 #include "OptimusComputeDataInterface.h"
 #include "ComputeFramework/ComputeDataProvider.h"
-#include "ChaosFlesh/FleshComponent.h"
+#include "ChaosFlesh/ChaosDeformableTetrahedralComponent.h"
 
 #include "DIFleshDeformer.generated.h"
 
@@ -31,6 +31,8 @@ class /*OPTIMUSCORE_API*/ UDIFleshDeformer : public UOptimusComputeDataInterface
 	GENERATED_BODY()
 
 public:
+	virtual ~UDIFleshDeformer();
+
 	//~ Begin UOptimusComputeDataInterface Interface
 	FString GetDisplayName() const override;
 	TArray<FOptimusCDIPinDefinition> GetPinDefinitions() const override;
@@ -49,6 +51,9 @@ public:
 
 	UPROPERTY(EditAnywhere, Category = Deformer, meta = (ShowOnlyInnerProperties))
 	FFleshDeformerParameters FleshDeformerParameters;
+
+private:
+	mutable TWeakObjectPtr<UDeformableTetrahedralComponent> ProducerComponent;
 };
 
 /** Compute Framework Data Provider for reading skeletal mesh. */
@@ -65,7 +70,7 @@ public:
 	TObjectPtr<UStaticMeshComponent> StaticMesh = nullptr;
 
 	UPROPERTY(EditAnywhere, BlueprintReadWrite, Category = Binding)
-	TObjectPtr<UFleshComponent> FleshMesh = nullptr;
+	TObjectPtr<UDeformableTetrahedralComponent> FleshMesh = nullptr;
 
 	UPROPERTY()
 	FFleshDeformerParameters FleshDeformerParameters;
@@ -79,7 +84,7 @@ public:
 class FDIFleshDeformerProviderProxy : public FComputeDataProviderRenderProxy
 {
 public:
-	FDIFleshDeformerProviderProxy(USkinnedMeshComponent* SkinnedMeshComponentIn, UFleshComponent* FleshComponentIn, const FFleshDeformerParameters& FleshDeformerParametersIn);
+	FDIFleshDeformerProviderProxy(USkinnedMeshComponent* SkinnedMeshComponentIn, UDeformableTetrahedralComponent* FleshComponentIn, const FFleshDeformerParameters& FleshDeformerParametersIn);
 
 	//~ Begin FComputeDataProviderRenderProxy Interface
 	bool IsValid(FValidationData const& InValidationData) const override;
@@ -93,7 +98,7 @@ private:
 
 	USkinnedMeshComponent* SkinnedMeshComponent = nullptr;
 	FSkeletalMeshObject* SkeletalMeshObject = nullptr;
-	UFleshComponent* FleshComponent = nullptr;
+	UDeformableTetrahedralComponent* FleshComponent = nullptr;
 
 	FFleshDeformerParameters FleshDeformerParameters;
 
