Fix indexing issue.

diff --git a/Engine/Plugins/Experimental/ChaosFlesh/Source/ChaosFlesh/Private/ChaosFlesh/TetrahedralCollection.cpp b/Engine/Plugins/Experimental/ChaosFlesh/Source/ChaosFlesh/Private/ChaosFlesh/TetrahedralCollection.cpp
--- a/Engine/Plugins/Experimental/ChaosFlesh/Source/ChaosFlesh/Private/ChaosFlesh/TetrahedralCollection.cpp
+++ b/Engine/Plugins/Experimental/ChaosFlesh/Source/ChaosFlesh/Private/ChaosFlesh/TetrahedralCollection.cpp
@@ -186,7 +186,7 @@ int32 FTetrahedralCollection::AppendGeometry(
 		}
 
 		const TArray<int32>& OtherIncidentElementsLocalIndex = Other.IncidentElementsLocalIndex[Idx];
-		TArray<int32>& ThisIncidentElementsLocalIndex = IncidentElementsLocalIndex[Idx];
+		TArray<int32>& ThisIncidentElementsLocalIndex = IncidentElementsLocalIndex[VerticesIndex + Idx];
 		// No offset necessary, just copy
 		ThisIncidentElementsLocalIndex = OtherIncidentElementsLocalIndex;
 	}
