ChaosFlesh matches SupportedTargetPlatforms for USDImporter

#rnx

diff --git a/Engine/Plugins/Experimental/ChaosFlesh/ChaosFlesh.uplugin b/Engine/Plugins/Experimental/ChaosFlesh/ChaosFlesh.uplugin
--- a/Engine/Plugins/Experimental/ChaosFlesh/ChaosFlesh.uplugin
+++ b/Engine/Plugins/Experimental/ChaosFlesh/ChaosFlesh.uplugin
@@ -64,7 +64,12 @@
 		},
         {
             "Name": "USDImporter",
-            "Enabled": true
+            "Enabled": true,
+			"SupportedTargetPlatforms": [
+				"Win64",
+				"Mac",
+				"Linux"
+			]
         }
 	]
 }
