ChaosFlesh : Renamed module to align with plugin name

diff --git a/Engine/Plugins/Experimental/ChaosFlesh/Source/ChaosFleshEditor/Private/ChaosFlesh/ChaosFleshEditorPlugin.cpp b/Engine/Plugins/Experimental/ChaosFlesh/Source/ChaosFleshEditor/Private/ChaosFlesh/ChaosFleshEditorPlugin.cpp
--- a/Engine/Plugins/Experimental/ChaosFlesh/Source/ChaosFleshEditor/Private/ChaosFlesh/ChaosFleshEditorPlugin.cpp
+++ b/Engine/Plugins/Experimental/ChaosFlesh/Source/ChaosFleshEditor/Private/ChaosFlesh/ChaosFleshEditorPlugin.cpp
@@ -88,7 +88,7 @@ void IChaosFleshEditorPlugin::ShutdownModule()
 	}
 }
 
-IMPLEMENT_MODULE(IChaosFleshEditorPlugin, FleshAssetEditor)
+IMPLEMENT_MODULE(IChaosFleshEditorPlugin, ChaosFleshEditor)
 
 
 #undef LOCTEXT_NAMESPACE
