ChaosFlesh: Muscle fiber visualization

diff --git a/Engine/Plugins/Experimental/ChaosFlesh/Source/ChaosFleshNodes/Private/Dataflow/ChaosFleshEngineAssetNodes.cpp b/Engine/Plugins/Experimental/ChaosFlesh/Source/ChaosFleshNodes/Private/Dataflow/ChaosFleshEngineAssetNodes.cpp
--- a/Engine/Plugins/Experimental/ChaosFlesh/Source/ChaosFleshNodes/Private/Dataflow/ChaosFleshEngineAssetNodes.cpp
+++ b/Engine/Plugins/Experimental/ChaosFlesh/Source/ChaosFleshNodes/Private/Dataflow/ChaosFleshEngineAssetNodes.cpp
@@ -24,6 +24,7 @@ namespace Dataflow
 		DATAFLOW_NODE_REGISTER_CREATION_FACTORY(FFleshAssetTerminalDataflowNode);
 		DATAFLOW_NODE_REGISTER_CREATION_FACTORY(FSetFleshDefaultPropertiesNode);
 		DATAFLOW_NODE_REGISTER_CREATION_FACTORY(FComputeFiberFieldNode);
+		DATAFLOW_NODE_REGISTER_CREATION_FACTORY(FVisualizeFiberFieldNode);
 		DATAFLOW_NODE_REGISTER_CREATION_FACTORY(FComputeIslandsNode);
 		DATAFLOW_NODE_REGISTER_CREATION_FACTORY(FGenerateOriginInsertionNode);
 		DATAFLOW_NODE_REGISTER_CREATION_FACTORY(FIsolateComponentNode);
@@ -353,6 +354,38 @@ void FComputeFiberFieldNode::Evaluate(Dataflow::FContext& Context, const FDatafl
 	}
 }
 
+void FVisualizeFiberFieldNode::Evaluate(Dataflow::FContext& Context, const FDataflowOutput* Out) const
+{
+	if (Out->IsA<FFieldCollection>(&VectorField))
+	{
+		FManagedArrayCollection InCollection = GetValue<FManagedArrayCollection>(Context, &Collection);
+		FFieldCollection OutVectorField = VectorField;
+		
+		if (TManagedArray<FVector3f>* Vertex = InCollection.FindAttribute<FVector3f>("Vertex", "Vertices"))
+		{
+			if (TManagedArray<FIntVector4>* Elements = InCollection.FindAttribute<FIntVector4>(FTetrahedralCollection::TetrahedronAttribute, FTetrahedralCollection::TetrahedralGroup))
+			{
+				if (TManagedArray<FVector3f>* FiberDirections = InCollection.FindAttribute<FVector3f>("FiberDirection", FTetrahedralCollection::TetrahedralGroup))
+				{
+					ensureMsgf(Elements->Num() == FiberDirections->Num(), TEXT("Fiber direction has different size than elements"));
+					for (int32 ElemIndex = 0; ElemIndex < Elements->Num(); ElemIndex++)
+					{
+						FVector3f VectorStart = { 0,0,0 };
+						for (int32 LocalIndex = 0; LocalIndex < 4; LocalIndex++)
+						{
+							VectorStart += (*Vertex)[(*Elements)[ElemIndex][LocalIndex]];
+						}
+						VectorStart /= float(4);
+						FVector3f VectorEnd = VectorStart + (*FiberDirections)[ElemIndex] * VectorScale;
+						OutVectorField.AddVectorToField(VectorStart, VectorEnd);
+					}
+				}
+			}
+		}
+		
+		Out->SetValue(MoveTemp(OutVectorField), Context);
+	}
+}
 
 TArray<int32> 
 FComputeFiberFieldNode::GetNonZeroIndices(const TArray<uint8>& Map) const
@@ -416,7 +449,7 @@ void FComputeIslandsNode::Evaluate(Dataflow::FContext& Context, const FDataflowO
 				Chaos::Utilities::FindConnectedRegions(Elements->GetConstArray(), ConnectedComponents);
 				TManagedArray<int32>& ComponentIndex = InCollection.ModifyAttribute<int32>("ComponentIndex", FGeometryCollection::VerticesGroup);
 
-				ComponentIndex.Fill(-1);
+				ComponentIndex.Fill(INDEX_NONE); //Isolated points will get index -1 
 
 				for (int32 i = 0; i < ConnectedComponents.Num(); i++)
 				{
@@ -426,7 +459,7 @@ void FComputeIslandsNode::Evaluate(Dataflow::FContext& Context, const FDataflowO
 						for (int32 ie = 0; ie < 4; ie++)
 						{
 							int32 ParticleIndex = (*Elements)[ElementIndex][ie];
-							if (ComponentIndex[ParticleIndex] == -1)
+							if (ComponentIndex[ParticleIndex] == INDEX_NONE)
 							{
 								ComponentIndex[ParticleIndex] = i;
 							}
@@ -644,4 +677,4 @@ void FGetSurfaceIndicesNode::Evaluate(Dataflow::FContext& Context, const FDatafl
 		}
 	}
 	SetValue(Context, MoveTemp(SurfaceIndicesLocal), &SurfaceIndicesOut);
-}
\ No newline at end of file
+}
diff --git a/Engine/Plugins/Experimental/ChaosFlesh/Source/ChaosFleshNodes/Public/Dataflow/ChaosFleshEngineAssetNodes.h b/Engine/Plugins/Experimental/ChaosFlesh/Source/ChaosFleshNodes/Public/Dataflow/ChaosFleshEngineAssetNodes.h
--- a/Engine/Plugins/Experimental/ChaosFlesh/Source/ChaosFleshNodes/Public/Dataflow/ChaosFleshEngineAssetNodes.h
+++ b/Engine/Plugins/Experimental/ChaosFlesh/Source/ChaosFleshNodes/Public/Dataflow/ChaosFleshEngineAssetNodes.h
@@ -3,12 +3,14 @@
 #pragma once 
 
 #include "CoreMinimal.h"
-#include "Dataflow/DataflowCore.h"
 #include "Chaos/Math/Poisson.h"
-#include "ChaosLog.h"
 #include "ChaosFlesh/TetrahedralCollection.h"
-#include "Dataflow/DataflowEngine.h"
 #include "ChaosFlesh/FleshAsset.h"
+#include "ChaosLog.h"
+
+#include "Dataflow/DataflowCore.h"
+#include "Dataflow/DataflowEngine.h"
+#include "Field/FieldSystemTypes.h"
 
 #include "ChaosFleshEngineAssetNodes.generated.h"
 
@@ -164,19 +166,47 @@ public:
 		const TArray<int32>& Insertion) const;
 };
 
+/**
+* Visualizes a muscle fiber direction per tetrahedron from a GeometryCollection containing tetrahedra.
+*/
+USTRUCT(meta = (DataflowFlesh))
+struct FVisualizeFiberFieldNode : public FDataflowNode
+{
+	GENERATED_USTRUCT_BODY()
+	DATAFLOW_NODE_DEFINE_INTERNAL(FVisualizeFiberFieldNode, "VisualizeFiberField", "Flesh", "")
+	DATAFLOW_NODE_RENDER_TYPE(FFieldCollection::StaticType(), "VectorField")
 
+public:
+	UPROPERTY(meta = (DataflowInput, DisplayName = "Collection"))
+	FManagedArrayCollection Collection;
+
+	UPROPERTY(EditAnywhere, Category = "Dataflow")
+	float VectorScale = 1.0;
+
+	UPROPERTY(meta = (DataflowOutput, DisplayName = "VectorField"))
+	FFieldCollection VectorField;
+	
+	FVisualizeFiberFieldNode(const Dataflow::FNodeParameters& InParam, FGuid InGuid = FGuid::NewGuid())
+		: FDataflowNode(InParam, InGuid)
+	{
+		RegisterInputConnection(&Collection);
+		RegisterOutputConnection(&VectorField);
+	}
+
+	virtual void Evaluate(Dataflow::FContext& Context, const FDataflowOutput* Out) const override;
+};
 
 USTRUCT(meta = (DataflowFlesh))
 struct FComputeIslandsNode : public FDataflowNode
 {
 	GENERATED_USTRUCT_BODY()
-		DATAFLOW_NODE_DEFINE_INTERNAL(FComputeIslandsNode, "ComputeIslands", "Flesh", "")
-		DATAFLOW_NODE_RENDER_TYPE(FGeometryCollection::StaticType(), "Collection")
+	DATAFLOW_NODE_DEFINE_INTERNAL(FComputeIslandsNode, "ComputeIslands", "Flesh", "")
+	DATAFLOW_NODE_RENDER_TYPE(FGeometryCollection::StaticType(), "Collection")
 
 public:
 
 	UPROPERTY(meta = (DataflowInput, DataflowOutput, DisplayName = "Collection", DataflowPassthrough = "Collection"))
-		FManagedArrayCollection Collection;
+	FManagedArrayCollection Collection;
 
 	FComputeIslandsNode(const Dataflow::FNodeParameters& InParam, FGuid InGuid = FGuid::NewGuid())
 		: FDataflowNode(InParam, InGuid)
@@ -200,22 +230,22 @@ public:
 	//typedef FManagedArrayCollection DataType;
 
 	UPROPERTY(meta = (DataflowInput, DataflowOutput, DisplayName = "Collection", DataflowPassthrough = "Collection"))
-		FManagedArrayCollection Collection;
+	FManagedArrayCollection Collection;
 
 	UPROPERTY(meta = (DataflowInput, DisplayName = "OriginIndicesIn"))
-		TArray<int32> OriginIndicesIn;
+	TArray<int32> OriginIndicesIn;
 
 	UPROPERTY(meta = (DataflowInput, DisplayName = "InsertionIndicesIn"))
-		TArray<int32> InsertionIndicesIn;
+	TArray<int32> InsertionIndicesIn;
 
 	UPROPERTY(meta = (DataflowOutput, DisplayName = "OriginIndicesOut"))
-		TArray<int32> OriginIndicesOut;
+	TArray<int32> OriginIndicesOut;
 
 	UPROPERTY(meta = (DataflowOutput, DisplayName = "InsertionIndicesOut"))
-		TArray<int32> InsertionIndicesOut;
+	TArray<int32> InsertionIndicesOut;
 	
 	UPROPERTY(EditAnywhere, Category = "Dataflow")
-		float Radius = float(1);
+	float Radius = float(1);
 
 	FGenerateOriginInsertionNode(const Dataflow::FNodeParameters& InParam, FGuid InGuid = FGuid::NewGuid())
 		: FDataflowNode(InParam, InGuid)
@@ -242,13 +272,13 @@ public:
 	//typedef FManagedArrayCollection DataType;
 
 	UPROPERTY(meta = (DataflowInput, DataflowOutput, DisplayName = "Collection", DataflowPassthrough = "Collection"))
-		FManagedArrayCollection Collection;
+	FManagedArrayCollection Collection;
 
 	UPROPERTY(EditAnywhere, Category = "Dataflow")
-		bool bDeleteHiddenFaces = false;
+	bool bDeleteHiddenFaces = false;
 
 	UPROPERTY(EditAnywhere, Category = "Dataflow")
-		FString TargetComponentIndex = "";
+	FString TargetComponentIndex = "";
 
 	FIsolateComponentNode(const Dataflow::FNodeParameters& InParam, FGuid InGuid = FGuid::NewGuid())
 		: FDataflowNode(InParam, InGuid)
@@ -269,13 +299,13 @@ struct FGetSurfaceIndicesNode : public FDataflowNode
 
 public:
 	UPROPERTY(meta = (DataflowInput, DisplayName = "Collection"))
-		FManagedArrayCollection Collection;
+	FManagedArrayCollection Collection;
 
 	UPROPERTY(meta = (DataflowInput, DisplayName = "GeometryGroupGuidsIn"))
-		TArray<FString> GeometryGroupGuidsIn;
+	TArray<FString> GeometryGroupGuidsIn;
 
 	UPROPERTY(meta = (DataflowOutput, DisplayName = "SurfaceIndicesOut"))
-		TArray<int32> SurfaceIndicesOut;
+	TArray<int32> SurfaceIndicesOut;
 
 	FGetSurfaceIndicesNode(const Dataflow::FNodeParameters& InParam, FGuid InGuid = FGuid::NewGuid())
 		: FDataflowNode(InParam, InGuid)
diff --git a/Engine/Plugins/Experimental/Dataflow/Source/DataflowEnginePlugin/Private/Dataflow/DataflowEngineRendering.cpp b/Engine/Plugins/Experimental/Dataflow/Source/DataflowEnginePlugin/Private/Dataflow/DataflowEngineRendering.cpp
--- a/Engine/Plugins/Experimental/Dataflow/Source/DataflowEnginePlugin/Private/Dataflow/DataflowEngineRendering.cpp
+++ b/Engine/Plugins/Experimental/Dataflow/Source/DataflowEnginePlugin/Private/Dataflow/DataflowEngineRendering.cpp
@@ -2,16 +2,16 @@
 
 #include "Dataflow/DataflowEngineRendering.h"
 
+#include "Dataflow/DataflowEnginePlugin.h"
 #include "Dataflow/DataflowRenderingFactory.h"
+#include "DynamicMesh/DynamicMesh3.h"
+#include "Field/FieldSystemTypes.h"
 #include "GeometryCollection/GeometryCollectionUtility.h"
 #include "GeometryCollection/Facades/CollectionRenderingFacade.h"
 #include "GeometryCollection/Facades/CollectionExplodedVectorFacade.h"
 #include "GeometryCollection/ManagedArrayCollection.h"
 #include "GeometryCollection/GeometryCollectionAlgo.h"
-#include "DynamicMesh/DynamicMesh3.h"
 #include "UDynamicMesh.h"
-#include "Dataflow/DataflowEnginePlugin.h"
-
 namespace Dataflow
 {
 
@@ -252,6 +252,54 @@ namespace Dataflow
 				}
 			});
 
+		/**
+		* DataflowNode (FFieldCollection) Rendering
+		*
+		*		@param Type : FName("FFieldCollection")
+
+		*		@param Outputs : {FFieldCollection : "VectorField"}
+		*/
+		FRenderingFactory::GetInstance()->RegisterOutput(FFieldCollection::StaticType(),
+			[](GeometryCollection::Facades::FRenderingFacade& RenderCollection, const Dataflow::FGraphRenderingState& State)
+			{
+				if (State.GetRenderOutputs().Num())
+				{
+					FName PrimaryOutput = State.GetRenderOutputs()[0]; // "VectorField"
+					if (PrimaryOutput.IsEqual(FName("VectorField")))
+					{
+						FFieldCollection Default;
+						const FFieldCollection& Collection = State.GetValue<FFieldCollection>(PrimaryOutput, Default);
+						TArray<TPair<FVector3f, FVector3f>> VectorField = Collection.GetVectorField();
+						const int32 NumVertices = 3 * VectorField.Num();
+						const int32 NumTriangles = VectorField.Num();
+
+						TArray<FVector3f> Vertices; Vertices.AddUninitialized(NumVertices);
+						TArray<FIntVector> Tris; Tris.AddUninitialized(NumTriangles);
+						TArray<FVector3f> VertexNormals; VertexNormals.AddUninitialized(NumVertices);
+						TArray<FLinearColor> VertexColors; VertexColors.AddUninitialized(NumVertices);
+
+						for (int32 i = 0; i < VectorField.Num(); i++)
+						{
+							
+							FVector3f Dir = VectorField[i].Value - VectorField[i].Key;
+							FVector3f DirAdd = Dir;
+							DirAdd.X += 1.f;
+							FVector3f OrthogonalDir = (Dir^ DirAdd).GetSafeNormal();
+							Tris[i] = FIntVector(3*i, 3*i+1, 3*i+2);
+							Vertices[3*i] = VectorField[i].Key;
+							Vertices[3*i+1] = VectorField[i].Value;
+							Vertices[3*i+2] = VectorField[i].Key + float(0.1) * Dir.Size() * OrthogonalDir;
+							FVector3f TriangleNormal = (OrthogonalDir ^ Dir).GetSafeNormal();
+							VertexNormals[3*i] = TriangleNormal;
+							VertexNormals[3*i+1] = TriangleNormal;
+							VertexNormals[3*i+2] = TriangleNormal;
+							VertexColors[i] = FLinearColor(IDataflowEnginePlugin::SurfaceColor);
+						}
+						RenderCollection.AddSurface(MoveTemp(Vertices), MoveTemp(Tris), MoveTemp(VertexNormals), MoveTemp(VertexColors));
+					}
+				}
+			});
+
 	}
 
 }
diff --git a/Engine/Source/Runtime/Experimental/Chaos/Public/Field/FieldSystemTypes.h b/Engine/Source/Runtime/Experimental/Chaos/Public/Field/FieldSystemTypes.h
--- a/Engine/Source/Runtime/Experimental/Chaos/Public/Field/FieldSystemTypes.h
+++ b/Engine/Source/Runtime/Experimental/Chaos/Public/Field/FieldSystemTypes.h
@@ -1,6 +1,7 @@
 // Copyright Epic Games, Inc. All Rights Reserved.
 #pragma once
 
+#include "GeometryCollection/ManagedArrayCollection.h"
 #include "UObject/ObjectMacros.h"
 
 #include "FieldSystemTypes.generated.h"
@@ -403,6 +404,38 @@ enum EFieldPhysicsDefaultFields : int
 	Field_EFieldPhysicsDefaultFields_Max                 UMETA(Hidden)
 };
 
+USTRUCT()
+struct CHAOS_API FFieldCollection : public FManagedArrayCollection
+{
+	GENERATED_USTRUCT_BODY()
+public:
+	static FName StaticType() { return FName("FFieldCollection"); }
+	FFieldCollection() 
+	{ 
+		AddGroup("VectorField"); 
+		AddAttribute<FVector3f>("Start", "VectorField");
+		AddAttribute<FVector3f>("End", "VectorField");
+	}
+	
+	void AddVectorToField(FVector3f Start, FVector3f End)
+	{
+		int32 Size = AddElements(1, "VectorField");
+		ModifyAttribute<FVector3f>("Start", "VectorField")[Size] = Start;
+		ModifyAttribute<FVector3f>("End", "VectorField")[Size] = End;
+	}
+	
+	TArray<TPair<FVector3f, FVector3f>> GetVectorField() const
+	{
+		TArray<TPair<FVector3f, FVector3f>> VectorField;
+		const TManagedArray<FVector3f>& VectorFieldStart = GetAttribute<FVector3f>("Start", "VectorField");
+		const TManagedArray<FVector3f>& VectorFieldEnd = GetAttribute<FVector3f>("End", "VectorField");
+		for (int32 i = 0; i < NumElements("VectorField"); i++)
+		{
+			VectorField.Add(TPair<FVector3f, FVector3f>(VectorFieldStart[i], VectorFieldEnd[i]));
+		}
+		return VectorField;
+	}
 
+};
 
 
