Correct USDImporter dependency omission.

diff --git a/Engine/Plugins/Experimental/ChaosFlesh/ChaosFlesh.uplugin b/Engine/Plugins/Experimental/ChaosFlesh/ChaosFlesh.uplugin
--- a/Engine/Plugins/Experimental/ChaosFlesh/ChaosFlesh.uplugin
+++ b/Engine/Plugins/Experimental/ChaosFlesh/ChaosFlesh.uplugin
@@ -68,6 +68,13 @@
 		{
 			"Name": "DeformerGraph",
 			"Enabled": true
-		}
+		},
+        {
+            "Name": "USDImporter",
+			"Enabled": true,
+			"SupportedTargetPlatforms": [
+				"Win64"
+			]
+        }
 	]
 }
