FleshComponentCacheAdapter: Make the cache directory browser relative to the project.  Add error logging.

diff --git a/Engine/Plugins/Experimental/ChaosCaching/Source/ChaosCaching/Public/Chaos/CacheManagerActor.h b/Engine/Plugins/Experimental/ChaosCaching/Source/ChaosCaching/Public/Chaos/CacheManagerActor.h
--- a/Engine/Plugins/Experimental/ChaosCaching/Source/ChaosCaching/Public/Chaos/CacheManagerActor.h
+++ b/Engine/Plugins/Experimental/ChaosCaching/Source/ChaosCaching/Public/Chaos/CacheManagerActor.h
@@ -108,8 +108,8 @@ struct FObservedComponent
 	UPROPERTY(EditAnywhere, BlueprintReadWrite, Category = "Caching")
 	bool bPlaybackEnabled;
 
-	/** USD cache directory, if supported. */
-	UPROPERTY(EditAnywhere, BlueprintReadWrite, Category = "Caching")
+	/** USD cache directory, if supported for this simulated structure type. */
+	UPROPERTY(EditAnywhere, BlueprintReadWrite, Category = "Caching", meta=(ContentDir))
 	FDirectoryPath USDCacheDirectory;
 
 	/** 
diff --git a/Engine/Plugins/Experimental/ChaosCachingUSD/Source/ChaosCachingUSD/Private/ChaosCachingUSD/Module.cpp b/Engine/Plugins/Experimental/ChaosCachingUSD/Source/ChaosCachingUSD/Private/ChaosCachingUSD/Module.cpp
--- a/Engine/Plugins/Experimental/ChaosCachingUSD/Source/ChaosCachingUSD/Private/ChaosCachingUSD/Module.cpp
+++ b/Engine/Plugins/Experimental/ChaosCachingUSD/Source/ChaosCachingUSD/Private/ChaosCachingUSD/Module.cpp
@@ -1,6 +1,7 @@
 // Copyright Epic Games, Inc. All Rights Reserved.
 
-#include "Modules/ModuleInterface.h"
+#include "ChaosCachingUSD/Module.h"
+
 #include "Modules/ModuleManager.h"
 #include "Interfaces/IPluginManager.h"
 #include "Misc/Paths.h"
@@ -8,14 +9,9 @@
 #include "UnrealUSDWrapper.h"
 #include "USDMemory.h"
 
+IMPLEMENT_MODULE_USD(FChaosCachingUSDModule, ChaosCachingUSD);
 
-class FChaosCachingUSDModule : public IModuleInterface
-{
-	//~ Begin IModuleInterface API
-	virtual void StartupModule() override;
-	//~ End IModuleInterface API
-
-};
+DEFINE_LOG_CATEGORY(LogChaosCacheUSD)
 
 void FChaosCachingUSDModule::StartupModule()
 {
@@ -34,6 +30,3 @@ void FChaosCachingUSDModule::StartupModule()
 	UnrealUSDWrapper::RegisterPlugins(ChaosCachingUSDResourcesDir);
 #endif // USE_USD_SDK
 }
-
-
-IMPLEMENT_MODULE_USD(FChaosCachingUSDModule, ChaosCachingUSD);
diff --git a/Engine/Plugins/Experimental/ChaosCachingUSD/Source/ChaosCachingUSD/Public/ChaosCachingUSD/Module.h b/Engine/Plugins/Experimental/ChaosCachingUSD/Source/ChaosCachingUSD/Public/ChaosCachingUSD/Module.h
new file mode 100644
--- /dev/null
+++ b/Engine/Plugins/Experimental/ChaosCachingUSD/Source/ChaosCachingUSD/Public/ChaosCachingUSD/Module.h
@@ -0,0 +1,15 @@
+// Copyright Epic Games, Inc. All Rights Reserved.
+
+#pragma once
+
+#include "Modules/ModuleInterface.h"
+
+DECLARE_LOG_CATEGORY_EXTERN(LogChaosCacheUSD, Verbose, All);
+
+class FChaosCachingUSDModule : public IModuleInterface
+{
+	//~ Begin IModuleInterface API
+	virtual void StartupModule() override;
+	//~ End IModuleInterface API
+
+};
diff --git a/Engine/Plugins/Experimental/ChaosFlesh/Source/ChaosFleshEngine/Private/ChaosCache/FleshComponentCacheAdapter.cpp b/Engine/Plugins/Experimental/ChaosFlesh/Source/ChaosFleshEngine/Private/ChaosCache/FleshComponentCacheAdapter.cpp
--- a/Engine/Plugins/Experimental/ChaosFlesh/Source/ChaosFleshEngine/Private/ChaosCache/FleshComponentCacheAdapter.cpp
+++ b/Engine/Plugins/Experimental/ChaosFlesh/Source/ChaosFleshEngine/Private/ChaosCache/FleshComponentCacheAdapter.cpp
@@ -19,6 +19,8 @@
 
 #endif // USE_USD_SDK && DO_USD_CACHING
 
+DEFINE_LOG_CATEGORY(LogChaosFleshCache)
+
 namespace Chaos
 {
 	FFleshCacheAdapterCVarParams CVarParams;
@@ -106,11 +108,18 @@ namespace Chaos
 						{
 							if (!UE::ChaosCachingUSD::WritePoints(MonolithStage, PrimPath, InTime, Points, Vels))
 							{
-								// TODO: Warn
+								UE_LOG(LogChaosFleshCache, Error,
+									TEXT("Failed to write points '%s' at time %g to file: '%s'"),
+									*PrimPath, InTime, *MonolithStage.GetRootLayer().GetDisplayName());
 								return;
 							}
 							// Save every n frames?
-							UE::ChaosCachingUSD::SaveStage(MonolithStage, MinTime, MaxTime);
+							if (!UE::ChaosCachingUSD::SaveStage(MonolithStage, MinTime, MaxTime))
+							{
+								UE_LOG(LogChaosFleshCache, Error,
+									TEXT("Failed to save file: '%s'"),
+									*MonolithStage.GetRootLayer().GetDisplayName());
+							}
 						}
 					}
 #else // USE_USD_SDK && DO_USD_CACHING
@@ -184,12 +193,20 @@ namespace Chaos
 						Prev != PrevV ||
 						Next != NextV)
 					{
+						UE_LOG(LogChaosFleshCache, Error,
+							TEXT("Inconsistent bracketing time samples for attributes '%s' and '%s' at frame %g from file: '%s'"),
+							*UE::ChaosCachingUSD::GetPointsAttrName(),
+							*UE::ChaosCachingUSD::GetVelocityAttrName(),
+							TargetTime, *MonolithStage.GetRootLayer().GetDisplayName());
 						return;
 					}
 
 					if (!UE::ChaosCachingUSD::ReadPoints(MonolithStage, PrimPath, Prev, Points0, Vels0) ||
 						Points0.size() != Vels0.size())
 					{
+						UE_LOG(LogChaosFleshCache, Error,
+							TEXT("Failed to read points '%s' at time %g from file: '%s'"),
+							*PrimPath, Prev, *MonolithStage.GetRootLayer().GetDisplayName());
 						return;
 					}
 
@@ -225,6 +242,9 @@ namespace Chaos
 						Points1.size() != Vels1.size() ||
 						Points0.size() != Points1.size())
 					{
+						UE_LOG(LogChaosFleshCache, Error,
+							TEXT("Failed to read points '%s' at time %g from file: '%s'"),
+							*PrimPath, Next, *MonolithStage.GetRootLayer().GetDisplayName());
 						return;
 					}
 					double Duration = Next - Prev;
@@ -370,6 +390,9 @@ namespace Chaos
 
 					if (!UE::ChaosCachingUSD::ReadPoints(MonolithStage, PrimPath, -TNumericLimits<double>::Max(), Points, Vels))
 					{
+						UE_LOG(LogChaosFleshCache, Error, 
+							TEXT("Failed to read points '%s' at time 'default' from file: '%s'"), 
+							*PrimPath, *MonolithStage.GetRootLayer().GetDisplayName());
 						return;
 					}
 
@@ -513,6 +536,7 @@ namespace Chaos
 				{
 					if (!PlatformFile.CreateDirectory(*CacheDir))
 					{
+						UE_LOG(LogChaosFleshCache, Error, TEXT("Failed to create output directory: '%s'"), *CacheDir);
 						return false;
 					}
 				}
@@ -521,6 +545,7 @@ namespace Chaos
 				const FFleshCollection* RestCollection = RestCollectionAsset ? RestCollectionAsset->GetCollection() : nullptr;
 				if (!RestCollection)
 				{
+					UE_LOG(LogChaosFleshCache, Error, TEXT("Failed to get rest collection from flesh component: '%s'"), *FleshComp->GetName());
 					return false;
 				}
 
@@ -546,6 +571,7 @@ namespace Chaos
 
 							if (!PlatformFile.MoveFile(*UniqueFilePath, *FilePath))
 							{
+								UE_LOG(LogChaosFleshCache, Error, TEXT("Failed to rename file from '%s' to '%s'."), *FilePath, *UniqueFilePath);
 								return false;
 							}
 						}
@@ -554,6 +580,7 @@ namespace Chaos
 					{
 						if (!PlatformFile.DeleteFile(*FilePath))
 						{
+							UE_LOG(LogChaosFleshCache, Error, TEXT("Failed to remove existing cache file: '%s'"), *FilePath);
 							return false;
 						}
 					}
@@ -564,10 +591,12 @@ namespace Chaos
 					}
 					if (!UE::ChaosCachingUSD::NewStage(FilePath, MonolithStage))
 					{
+						UE_LOG(LogChaosFleshCache, Error, TEXT("Failed to create new USD file: '%s'"), *FilePath);
 						return false;
 					}
 					if (!UE::ChaosCachingUSD::WriteTetMesh(MonolithStage, PrimPath, *RestCollection))
 					{
+						UE_LOG(LogChaosFleshCache, Error, TEXT("Failed to write tetrahedron mesh '%s' to USD file: '%s'"), *PrimPath, *FilePath);
 						return false;
 					}
 				}
@@ -616,15 +645,20 @@ namespace Chaos
 						}
 						if (!UE::ChaosCachingUSD::OpenStage(FilePath, MonolithStage))
 						{
+							UE_LOG(LogChaosFleshCache, Error, TEXT("Failed to open USD cache file: '%s'"), *FilePath);
 							return false;
 						}
 					}
 					else
 					{
+						UE_LOG(LogChaosFleshCache, Error, TEXT("USD cache file not found: '%s'"), *FilePath);
 						return false;
 					}
 				}
-
+				if (!bUseMonolith)
+				{
+					UE_LOG(LogChaosFleshCache, Warning, TEXT("No USD file structure selected (bMonolith)."));
+				}
 #endif // USE_USD_SDK && DO_USD_CACHING
 
 			}
diff --git a/Engine/Plugins/Experimental/ChaosFlesh/Source/ChaosFleshEngine/Public/ChaosCache/FleshComponentCacheAdapter.h b/Engine/Plugins/Experimental/ChaosFlesh/Source/ChaosFleshEngine/Public/ChaosCache/FleshComponentCacheAdapter.h
--- a/Engine/Plugins/Experimental/ChaosFlesh/Source/ChaosFleshEngine/Public/ChaosCache/FleshComponentCacheAdapter.h
+++ b/Engine/Plugins/Experimental/ChaosFlesh/Source/ChaosFleshEngine/Public/ChaosCache/FleshComponentCacheAdapter.h
@@ -17,6 +17,8 @@
 
 #endif // USE_USD_SDK && DO_USD_CACHING
 
+DECLARE_LOG_CATEGORY_EXTERN(LogChaosFleshCache, Verbose, All);
+
 namespace Chaos::Softs
 {
 	class FDeformableSolver;
