Get UE relative paths right for USD files.  Add more logging.

diff --git a/Engine/Plugins/Experimental/ChaosCachingUSD/Source/ChaosCachingUSD/Private/ChaosCachingUSD/Operations.cpp b/Engine/Plugins/Experimental/ChaosCachingUSD/Source/ChaosCachingUSD/Private/ChaosCachingUSD/Operations.cpp
--- a/Engine/Plugins/Experimental/ChaosCachingUSD/Source/ChaosCachingUSD/Private/ChaosCachingUSD/Operations.cpp
+++ b/Engine/Plugins/Experimental/ChaosCachingUSD/Source/ChaosCachingUSD/Private/ChaosCachingUSD/Operations.cpp
@@ -97,7 +97,21 @@ UE::ChaosCachingUSD::SaveStage(UE::FUsdStage& Stage, const double FirstFrame, co
 			Stage.SetEndTimeCode(LastFrame);
 		}
 	}
-	return Stage.GetRootLayer().Save();
+	if (UE::FSdfLayer RootLayer = Stage.GetRootLayer())
+	{
+		UE_LOG(LogUsd, Log, TEXT("Saving USD file: '%s'"), *RootLayer.GetRealPath());
+		if (!RootLayer.Save())
+		{
+			UE_LOG(LogUsd, Error, TEXT("Failed to save USD file: '%s'"), *RootLayer.GetRealPath());
+			return false;
+		}
+		return true;
+	}
+	else
+	{
+		UE_LOG(LogUsd, Error, TEXT("USD Stage has no root layer."));
+		return false;
+	}
 }
 
 bool
@@ -446,8 +460,7 @@ UE::ChaosCachingUSD::WriteTetMesh(
 			*PrimPath, *Stage.GetRootLayer().GetRealPath());
 	}
 
-	UE::ChaosCachingUSD::SaveStage(Stage, -TNumericLimits<double>::Max(), -TNumericLimits<double>::Max());
-	return true;
+	return UE::ChaosCachingUSD::SaveStage(Stage, -TNumericLimits<double>::Max(), -TNumericLimits<double>::Max());
 }
 
 bool
diff --git a/Engine/Plugins/Experimental/ChaosFlesh/Source/ChaosFleshEngine/Private/ChaosCache/FleshComponentCacheAdapter.cpp b/Engine/Plugins/Experimental/ChaosFlesh/Source/ChaosFleshEngine/Private/ChaosCache/FleshComponentCacheAdapter.cpp
--- a/Engine/Plugins/Experimental/ChaosFlesh/Source/ChaosFleshEngine/Private/ChaosCache/FleshComponentCacheAdapter.cpp
+++ b/Engine/Plugins/Experimental/ChaosFlesh/Source/ChaosFleshEngine/Private/ChaosCache/FleshComponentCacheAdapter.cpp
@@ -478,11 +478,16 @@ namespace Chaos
 
 	FString GetCacheDirectory(const FObservedComponent& InObserved)
 	{
+		// USDCacheDirectory is relative to the content dir, but with "/Game" rather than just "/" or some relative path.
 		FString CacheDir = InObserved.USDCacheDirectory.Path;
 		if (CacheDir.IsEmpty())
 		{
-			FString ProjectDir = FPaths::ProjectContentDir();
-			CacheDir = FPaths::Combine(ProjectDir, FString(TEXT("SimCache")));
+			CacheDir = FString(TEXT("SimCache"));
+		}
+		FPaths::NormalizeDirectoryName(CacheDir);
+		if (CacheDir.StartsWith(FString(TEXT("/Game"))))
+		{
+			CacheDir = FPaths::Combine(FPaths::ProjectContentDir(), CacheDir.RightChop(5));
 		}
 		return CacheDir;
 	}
@@ -534,7 +539,7 @@ namespace Chaos
 				IPlatformFile& PlatformFile = FileManager.GetPlatformFile();
 				if (!PlatformFile.DirectoryExists(*CacheDir))
 				{
-					if (!PlatformFile.CreateDirectory(*CacheDir))
+					if (!PlatformFile.CreateDirectoryTree(*CacheDir))
 					{
 						UE_LOG(LogChaosFleshCache, Error, TEXT("Failed to create output directory: '%s'"), *CacheDir);
 						return false;
