Fix indexing issue in FAddKinematicParticlesDataflowNode.  Skip vertices with BoneMap == -1 when computing bounding boxes.

diff --git a/Engine/Plugins/Experimental/ChaosFlesh/Source/ChaosFlesh/Private/ChaosFlesh/TetrahedralCollection.cpp b/Engine/Plugins/Experimental/ChaosFlesh/Source/ChaosFlesh/Private/ChaosFlesh/TetrahedralCollection.cpp
--- a/Engine/Plugins/Experimental/ChaosFlesh/Source/ChaosFlesh/Private/ChaosFlesh/TetrahedralCollection.cpp
+++ b/Engine/Plugins/Experimental/ChaosFlesh/Source/ChaosFlesh/Private/ChaosFlesh/TetrahedralCollection.cpp
@@ -121,10 +121,14 @@ void FTetrahedralCollection::UpdateBoundingBox()
 		// Compute BoundingBox
 		for (int32 Idx = 0; Idx < Vertex.Num(); ++Idx)
 		{
-			int32 TransformIndexValue = BoneMap[Idx];
-			if (TransformToGeometryIndex[TransformIndexValue] != INDEX_NONE)
+			const int32 TransformIndexValue = BoneMap[Idx];
+			if (TransformIndexValue != INDEX_NONE)
 			{
-				BoundingBox[TransformToGeometryIndex[TransformIndexValue]] += FVector(Vertex[Idx]);
+				const int32 GeometryIndex = TransformToGeometryIndex[TransformIndexValue];
+				if (GeometryIndex != INDEX_NONE)
+				{
+					BoundingBox[GeometryIndex] += FVector(Vertex[Idx]);
+				}
 			}
 		}
 	}
diff --git a/Engine/Plugins/Experimental/ChaosFlesh/Source/ChaosFleshNodes/Private/Dataflow/ChaosFleshPositionTargetInitializationNodes.cpp b/Engine/Plugins/Experimental/ChaosFlesh/Source/ChaosFleshNodes/Private/Dataflow/ChaosFleshPositionTargetInitializationNodes.cpp
--- a/Engine/Plugins/Experimental/ChaosFlesh/Source/ChaosFleshNodes/Private/Dataflow/ChaosFleshPositionTargetInitializationNodes.cpp
+++ b/Engine/Plugins/Experimental/ChaosFlesh/Source/ChaosFleshNodes/Private/Dataflow/ChaosFleshPositionTargetInitializationNodes.cpp
@@ -103,12 +103,14 @@ void FAddKinematicParticlesDataflowNode::Evaluate(Dataflow::FContext& Context, c
 						BranchIndices.Add(IndexValue);
 					}
 
+					// Add standalone particles, not bound to a transform group - so for these particles BoneMap = INDEX_NONE.
 					int32 ParticleIndex = InCollection.AddElements(BranchIndices.Num(), FGeometryCollection::VerticesGroup);
 					TManagedArray<FVector3f>& CurrentVertices = InCollection.ModifyAttribute<FVector3f>("Vertex", FGeometryCollection::VerticesGroup);
 
 					for (int32 Index = 0; Index < BranchIndices.Num(); Index++)
 					{
-						FVector3f BonePosition(ComponentPose[BranchIndices[Index]].GetTranslation());
+						const int32 BoneIndex = BranchIndices[Index];
+						FVector3f BonePosition(ComponentPose[BoneIndex].GetTranslation());
 
 						CurrentVertices[ParticleIndex+Index] = BonePosition;
 
@@ -119,11 +121,10 @@ void FAddKinematicParticlesDataflowNode::Evaluate(Dataflow::FContext& Context, c
 						BoundWeights.Add(1.0);
 						TargetIndices.Emplace(ParticleIndex + Index);
 
-
 						if (BoundVerts.Num())
 						{
 							GeometryCollection::Facades::FKinematicBindingFacade Kinematics(InCollection);
-							Kinematics.AddKinematicBinding(Kinematics.SetBoneBindings(Index, BoundVerts, BoundWeights));
+							Kinematics.AddKinematicBinding(Kinematics.SetBoneBindings(BoneIndex, BoundVerts, BoundWeights));
 						}
 					}
 
diff --git a/Engine/Source/Runtime/Experimental/Chaos/Private/GeometryCollection/GeometryCollection.cpp b/Engine/Source/Runtime/Experimental/Chaos/Private/GeometryCollection/GeometryCollection.cpp
--- a/Engine/Source/Runtime/Experimental/Chaos/Private/GeometryCollection/GeometryCollection.cpp
+++ b/Engine/Source/Runtime/Experimental/Chaos/Private/GeometryCollection/GeometryCollection.cpp
@@ -1002,8 +1002,15 @@ void FGeometryCollection::UpdateBoundingBox(FManagedArrayCollection& InCollectio
 		// Compute BoundingBox
 		for (int32 Idx = 0; Idx < Vertex.Num(); ++Idx)
 		{
-			int32 TransformIndexValue = BoneMap[Idx];
-			BoundingBox[TransformToGeometryIndex[TransformIndexValue]] += FVector(Vertex[Idx]);
+			const int32 TransformIndexValue = BoneMap[Idx];
+			if (TransformIndexValue != INDEX_NONE)
+			{
+				const int32 GeometryIndex = TransformToGeometryIndex[TransformIndexValue];
+				if (GeometryIndex != INDEX_NONE)
+				{
+					BoundingBox[GeometryIndex] += FVector(Vertex[Idx]);
+				}
+			}
 		}
 	}
 }
