Chaos : define fixes
- Use if instead of ifdef for WITH_EDITOR

diff --git a/Engine/Plugins/Experimental/ChaosFlesh/Source/ChaosFleshEngine/Private/ChaosFlesh/ChaosDeformableGPUBuffers.cpp b/Engine/Plugins/Experimental/ChaosFlesh/Source/ChaosFleshEngine/Private/ChaosFlesh/ChaosDeformableGPUBuffers.cpp
--- a/Engine/Plugins/Experimental/ChaosFlesh/Source/ChaosFleshEngine/Private/ChaosFlesh/ChaosDeformableGPUBuffers.cpp
+++ b/Engine/Plugins/Experimental/ChaosFlesh/Source/ChaosFleshEngine/Private/ChaosFlesh/ChaosDeformableGPUBuffers.cpp
@@ -193,7 +193,7 @@ uint32 FIndexArrayBufferWithSRV::GetBufferSize() const
 	return BufferSize;
 }
 
-#ifdef WITH_EDITOR
+#if WITH_EDITOR
 uint32 UnpackUInt(const uint32* Array, const uint32 Index, const uint32 Stride, const uint32 ArraySize)
 {
 	uint32 ValsPerEntry = 4 / Stride;
@@ -310,7 +310,7 @@ void FIndexArrayBufferWithSRV::InitRHI(FRHICommandListBase& RHICmdList)
 			}
         }
 
-#ifdef WITH_EDITOR
+#if WITH_EDITOR
 		if (Local::CVarParam.bTestUnpacking)
 		{
 			uint32* Indices32Bit = reinterpret_cast<uint32*>(Buffer);
diff --git a/Engine/Plugins/Experimental/ChaosFlesh/Source/ChaosFleshEngine/Private/ChaosFlesh/ChaosFleshDeformerBufferManager.cpp b/Engine/Plugins/Experimental/ChaosFlesh/Source/ChaosFleshEngine/Private/ChaosFlesh/ChaosFleshDeformerBufferManager.cpp
--- a/Engine/Plugins/Experimental/ChaosFlesh/Source/ChaosFleshEngine/Private/ChaosFlesh/ChaosFleshDeformerBufferManager.cpp
+++ b/Engine/Plugins/Experimental/ChaosFlesh/Source/ChaosFleshEngine/Private/ChaosFlesh/ChaosFleshDeformerBufferManager.cpp
@@ -286,7 +286,7 @@ bool FChaosFleshDeformableGPUManager::FBindingsBuffer::InitData(
 	}
 	Chaos::Softs::Private::InitBuffer(VerticesBuffer, &SparseVertices, Chaos::Softs::Private::NullFloatBuffer);
 
-#ifdef WITH_EDITOR
+#if WITH_EDITOR
 	if (Local::CVarParams.bTestSparseMappings)
 	{
 		TMap<int32, int32> DenseToSparseIndices;
