Fix Win64 only module reference

#rnx

diff --git a/Engine/Plugins/Experimental/ChaosFlesh/Source/ChaosFleshEngine/ChaosFleshEngine.Build.cs b/Engine/Plugins/Experimental/ChaosFlesh/Source/ChaosFleshEngine/ChaosFleshEngine.Build.cs
--- a/Engine/Plugins/Experimental/ChaosFlesh/Source/ChaosFleshEngine/ChaosFleshEngine.Build.cs
+++ b/Engine/Plugins/Experimental/ChaosFlesh/Source/ChaosFleshEngine/ChaosFleshEngine.Build.cs
@@ -38,7 +38,6 @@ namespace UnrealBuildTool.Rules
 					"CoreUObject",
 					"Chaos",
 					"ChaosCaching",
-					"ChaosCachingUSD",
 					"ChaosFlesh",
 					"DataflowCore",
 					"DataflowEngine",
@@ -50,12 +49,25 @@ namespace UnrealBuildTool.Rules
 					"RenderCore",
                     "RHI",
 					"Renderer",
-					"UnrealUSDWrapper",
-					"USDClasses",
-					"USDUtilities",
 				}
 				);
 
+			if (Target.Platform == UnrealTargetPlatform.Win64)
+			{
+				PrivateDependencyModuleNames.AddRange(
+					new string[]
+					{
+						"ChaosCachingUSD",
+						"UnrealUSDWrapper",
+						"USDClasses",
+						"USDUtilities",
+					});
+			}
+			else
+			{
+				PrivateDefinitions.Add("USE_USD_SDK=0");
+			}
+
 			DynamicallyLoadedModuleNames.AddRange(
 				new string[]
 				{
