ChaosFlesh : Expose some of the parameters to blueprints.

diff --git a/Engine/Plugins/Experimental/ChaosFlesh/Source/ChaosFleshEngine/Public/ChaosFlesh/ChaosDeformableGameplayComponent.h b/Engine/Plugins/Experimental/ChaosFlesh/Source/ChaosFleshEngine/Public/ChaosFlesh/ChaosDeformableGameplayComponent.h
--- a/Engine/Plugins/Experimental/ChaosFlesh/Source/ChaosFleshEngine/Public/ChaosFlesh/ChaosDeformableGameplayComponent.h
+++ b/Engine/Plugins/Experimental/ChaosFlesh/Source/ChaosFleshEngine/Public/ChaosFlesh/ChaosDeformableGameplayComponent.h
@@ -12,38 +12,38 @@
 class UBoxComponent;
 class UStaticMeshComponent;
 
-USTRUCT()
+USTRUCT(BlueprintType)
 struct FRigBoundRayCasts
 {
 	GENERATED_USTRUCT_BODY()
 
-	UPROPERTY(EditAnywhere, Category = "Physics")
+	UPROPERTY(EditAnywhere, BlueprintReadWrite, Category = "Physics")
 		bool bEnableRigBoundRaycasts = false;
 
-	UPROPERTY(EditAnywhere, Category = "Physics")
+	UPROPERTY(EditAnywhere, BlueprintReadWrite, Category = "Physics")
 		int32 MaxNumTests = 5;
 
-	UPROPERTY(EditAnywhere, Category = "Physics")
+	UPROPERTY(EditAnywhere, BlueprintReadWrite, Category = "Physics")
 		bool bTestDownOnly = false;
 
-	UPROPERTY(EditAnywhere, Category = "Physics")
+	UPROPERTY(EditAnywhere, BlueprintReadWrite, Category = "Physics")
 		float TestRange = 0.5f;
 	
 	/** Objects to skip if hit during \c DetectEnvironmentCollisions(). */
 	UPROPERTY(EditAnywhere, Category = "Physics")
 	TArray<TWeakObjectPtr<UPrimitiveComponent>> EnvironmentCollisionsSkipList;
 
-	UPROPERTY(EditAnywhere, Category = "Physics")
-		TEnumAsByte<ECollisionChannel> CollisionChannel = ECC_WorldStatic;
+	UPROPERTY(EditAnywhere, BlueprintReadWrite, Category = "Physics")
+	TEnumAsByte<ECollisionChannel> CollisionChannel = ECC_WorldStatic;
 };
 
 
-USTRUCT()
+USTRUCT(BlueprintType)
 struct FGameplayColllisions
 {
 	GENERATED_USTRUCT_BODY()
 
-	UPROPERTY(EditAnywhere, Category = "Physics")
+	UPROPERTY(EditAnywhere, BlueprintReadWrite, Category = "Physics")
 	FRigBoundRayCasts RigBoundRayCasts;
 
 };
@@ -82,7 +82,7 @@ public:
 	void DetectEnvironmentCollisions(const int32 MaxNumTests = 100, const bool bTestDownOnly = true, const float TestRange = 0.0, const ECollisionChannel CollisionChannel = ECollisionChannel::ECC_PhysicsBody);
 
 
-	UPROPERTY(EditAnywhere, Category = "Physics", DisplayName = "Collisions", meta = (DisplayPriority = 10))
+	UPROPERTY(EditAnywhere, BlueprintReadWrite, Category = "Physics", DisplayName = "Collisions", meta = (DisplayPriority = 10))
 	FGameplayColllisions GameplayColllisions;
 
 
