ChaosFlesh : Cull unused vertices from collection.
#perfight 65bd394e8e045c9abd2ca106

diff --git a/Engine/Plugins/Experimental/ChaosFlesh/Source/ChaosFleshNodes/Private/Dataflow/ChaosFleshTetrahedralNodes.cpp b/Engine/Plugins/Experimental/ChaosFlesh/Source/ChaosFleshNodes/Private/Dataflow/ChaosFleshTetrahedralNodes.cpp
--- a/Engine/Plugins/Experimental/ChaosFlesh/Source/ChaosFleshNodes/Private/Dataflow/ChaosFleshTetrahedralNodes.cpp
+++ b/Engine/Plugins/Experimental/ChaosFlesh/Source/ChaosFleshNodes/Private/Dataflow/ChaosFleshTetrahedralNodes.cpp
@@ -21,6 +21,7 @@
 #include "Generate/IsosurfaceStuffing.h"
 #include "GeometryCollection/ManagedArrayCollection.h"
 #include "GeometryCollection/Facades/CollectionTetrahedralMetricsFacade.h"
+#include "GeometryCollection/GeometryCollectionAlgo.h"
 #include "MeshDescription.h"
 #include "MeshDescriptionToDynamicMesh.h"
 #include "Rendering/SkeletalMeshLODImporterData.h"
@@ -304,6 +305,11 @@ void FGenerateTetrahedralCollectionDataflowNodes::Evaluate(Dataflow::FContext& C
 				{
 					if (CollectionBuffer[i]->NumElements(FGeometryCollection::VerticesGroup))
 					{
+						TSet<int32> VertexToDeleteSet;
+						GeometryCollectionAlgo::ComputeStaleVertices(CollectionBuffer[i].Get(), VertexToDeleteSet);
+						TArray<int32> SortedVertices = VertexToDeleteSet.Array(); SortedVertices.Sort();
+						if (VertexToDeleteSet.Num()) CollectionBuffer[i]->RemoveElements(FGeometryCollection::VerticesGroup, SortedVertices);
+
 						int32 GeomIndex = InCollection->AppendGeometry(*CollectionBuffer[i].Get());
 					}
 				}
diff --git a/Engine/Source/Runtime/Experimental/Chaos/Private/GeometryCollection/GeometryCollectionAlgo.cpp b/Engine/Source/Runtime/Experimental/Chaos/Private/GeometryCollection/GeometryCollectionAlgo.cpp
--- a/Engine/Source/Runtime/Experimental/Chaos/Private/GeometryCollection/GeometryCollectionAlgo.cpp
+++ b/Engine/Source/Runtime/Experimental/Chaos/Private/GeometryCollection/GeometryCollectionAlgo.cpp
@@ -1435,6 +1435,20 @@ namespace GeometryCollectionAlgo
 			VertexInFaceArray[IndicesArray[IdxFace].Z]++;
 		}
 
+		if (const TManagedArray<FIntVector4> *TetArrayPtr = 
+			GeometryCollection->FindAttributeTyped<FIntVector4>("Tetrahedron", "Tetrahedral"))
+		{
+			int32 NumTetrahedron = GeometryCollection->NumElements("Tetrahedral");
+			const TManagedArray<FIntVector4>& TetArray = *TetArrayPtr;
+			for (int32 TetIdx = 0; TetIdx < NumTetrahedron; ++TetIdx)
+			{
+				VertexInFaceArray[TetArray[TetIdx].X]++;
+				VertexInFaceArray[TetArray[TetIdx].Y]++;
+				VertexInFaceArray[TetArray[TetIdx].Z]++;
+				VertexInFaceArray[TetArray[TetIdx].W]++;
+			}
+		}
+
 		for (int32 IdxVertex = 0; IdxVertex < NumVertices; ++IdxVertex)
 		{
 			if (VertexInFaceArray[IdxVertex] == 0)
