ChaosFlesh : Fix for CIS

diff --git a/Engine/Plugins/Experimental/ChaosFlesh/Source/ChaosFleshNodes/Public/Dataflow/ChaosFleshImportGEO.h b/Engine/Plugins/Experimental/ChaosFlesh/Source/ChaosFleshNodes/Public/Dataflow/ChaosFleshImportGEO.h
--- a/Engine/Plugins/Experimental/ChaosFlesh/Source/ChaosFleshNodes/Public/Dataflow/ChaosFleshImportGEO.h
+++ b/Engine/Plugins/Experimental/ChaosFlesh/Source/ChaosFleshNodes/Public/Dataflow/ChaosFleshImportGEO.h
@@ -53,7 +53,7 @@ public:
 
 	/** Output value. */
 	UPROPERTY(EditAnywhere, Category = "Dataflow", meta = (DataflowOutput))
-	int32 Value;
+	int32 Value = 0;
 
 	FExtractGEOInt(const Dataflow::FNodeParameters& InParam, FGuid InGuid = FGuid::NewGuid())
 		: FDataflowNode(InParam, InGuid)
