ChaosFlesh : Fix out of bounds access.
#perflight 65cfc419334c09b9283337ae

diff --git a/Engine/Plugins/Experimental/ChaosFlesh/Source/ChaosFlesh/Private/ChaosFlesh/TetrahedralCollection.cpp b/Engine/Plugins/Experimental/ChaosFlesh/Source/ChaosFlesh/Private/ChaosFlesh/TetrahedralCollection.cpp
--- a/Engine/Plugins/Experimental/ChaosFlesh/Source/ChaosFlesh/Private/ChaosFlesh/TetrahedralCollection.cpp
+++ b/Engine/Plugins/Experimental/ChaosFlesh/Source/ChaosFlesh/Private/ChaosFlesh/TetrahedralCollection.cpp
@@ -95,12 +95,15 @@ void FTetrahedralCollection::Init(
 
 		Super::Init(Collection, RawVertexArray, RawIndicesArray, bReverseVertexOrder);
 
-		Collection->TetrahedronStart[0] = Collection->Tetrahedron.Num();
-		Collection->TetrahedronCount[0] = Elements.Num();
-		Collection->AddElements(Elements.Num(), FTetrahedralCollection::TetrahedralGroup);
-		for (int i = 0; i < Elements.Num(); i++)
+		if (Collection->NumElements(FTetrahedralCollection::GeometryGroup))
 		{
-			Collection->Tetrahedron[i] = Elements[i];
+			Collection->TetrahedronStart[0] = Collection->Tetrahedron.Num();
+			Collection->TetrahedronCount[0] = Elements.Num();
+			Collection->AddElements(Elements.Num(), FTetrahedralCollection::TetrahedralGroup);
+			for (int i = 0; i < Elements.Num(); i++)
+			{
+				Collection->Tetrahedron[i] = Elements[i];
+			}
 		}
 
 		// Init aux structures that depend on topology.
