ChaosFlesh : Command to convert a simulated cache with mesh attachments to a GeometryCache. (thanks ryank)

diff --git a/Engine/Plugins/Experimental/ChaosCaching/Source/ChaosCaching/Public/Chaos/CacheManagerActor.h b/Engine/Plugins/Experimental/ChaosCaching/Source/ChaosCaching/Public/Chaos/CacheManagerActor.h
--- a/Engine/Plugins/Experimental/ChaosCaching/Source/ChaosCaching/Public/Chaos/CacheManagerActor.h
+++ b/Engine/Plugins/Experimental/ChaosCaching/Source/ChaosCaching/Public/Chaos/CacheManagerActor.h
@@ -60,6 +60,7 @@ struct FObservedComponent
 		bIsSimulating = OtherComponent.bIsSimulating;
 		bHasNotifyBreaks = OtherComponent.bHasNotifyBreaks;
 		bPlaybackEnabled = OtherComponent.bPlaybackEnabled;
+		USDCacheDirectory = OtherComponent.USDCacheDirectory;
 		bTriggered = OtherComponent.bTriggered;
 		AbsoluteTime = OtherComponent.AbsoluteTime;
 		TimeSinceTrigger = OtherComponent.TimeSinceTrigger;
@@ -75,6 +76,7 @@ struct FObservedComponent
 		bIsSimulating = OtherComponent.bIsSimulating;
 		bHasNotifyBreaks = OtherComponent.bHasNotifyBreaks;
 		bPlaybackEnabled = OtherComponent.bPlaybackEnabled;
+		USDCacheDirectory = OtherComponent.USDCacheDirectory;
 		bTriggered = OtherComponent.bTriggered;
 		AbsoluteTime = OtherComponent.AbsoluteTime;
 		TimeSinceTrigger = OtherComponent.TimeSinceTrigger;
diff --git a/Engine/Plugins/Experimental/ChaosCachingUSD/Source/ChaosCachingUSD/Private/ChaosCachingUSD/Operations.cpp b/Engine/Plugins/Experimental/ChaosCachingUSD/Source/ChaosCachingUSD/Private/ChaosCachingUSD/Operations.cpp
--- a/Engine/Plugins/Experimental/ChaosCachingUSD/Source/ChaosCachingUSD/Private/ChaosCachingUSD/Operations.cpp
+++ b/Engine/Plugins/Experimental/ChaosCachingUSD/Source/ChaosCachingUSD/Private/ChaosCachingUSD/Operations.cpp
@@ -10,6 +10,8 @@
 #include "GeometryCollection/GeometryCollection.h"
 #include "GeometryCollection/ManagedArray.h"
 #include "GeometryCollection/ManagedArrayCollection.h"
+#include "HAL/PlatformFile.h"
+#include "HAL/PlatformFileManager.h"
 #include "Misc/Paths.h"
 
 #include "UnrealUSDWrapper.h"
@@ -91,8 +93,13 @@ UE::ChaosCachingUSD::OpenStage(const FString& StageName, UE::FUsdStage& UsdStage
 {
 	if (!FPaths::FileExists(StageName))
 	{
-		UE_LOG(LogUsd, Error, TEXT("File not found: '%s'"), *StageName);
-		return false;
+		FPlatformFileManager& FileManager = FPlatformFileManager::Get();
+		IPlatformFile& PlatformFile = FileManager.GetPlatformFile();
+		if (!PlatformFile.FileExists(*StageName))
+		{
+			UE_LOG(LogUsd, Error, TEXT("File not found: '%s'"), *StageName);
+			return false;
+		}
 	}
 
 	// USD caches all stages you open/create, unless you tell it not to.
@@ -748,6 +755,33 @@ UE::ChaosCachingUSD::ReadPoints(
 	return true;
 }
 
+bool 
+UE::ChaosCachingUSD::ReadPoints(
+	const UE::FUsdStage& Stage,
+	const FString& PrimPath,
+	const FString& AttrPath,
+	const double Time,
+	TArray<Chaos::TVector<Chaos::FRealSingle, 3>>& Points)
+{
+	FScopedUsdAllocs UsdAllocs; // Use USD memory allocator
+	pxr::VtArray<pxr::GfVec3f> VtPoints;
+	if (ReadPoints(Stage, PrimPath, AttrPath, Time, VtPoints))
+	{
+		{
+			FScopedUnrealAllocs UEAllocs; // Use UE memory allocator
+			Points.SetNum(VtPoints.size());
+		}
+		int32 i = 0;
+		for (pxr::VtArray<pxr::GfVec3f>::const_iterator it = VtPoints.cbegin(), itEnd = VtPoints.cend(); it != itEnd; ++it)
+		{
+			const pxr::GfVec3f& Pt = *it;
+			Points[i++].Set(Pt[0], Pt[1], Pt[2]);
+		}
+		return true;
+	}
+	return false;
+}
+
 bool
 UE::ChaosCachingUSD::ReadPoints(
 	const UE::FUsdStage& Stage,
diff --git a/Engine/Plugins/Experimental/ChaosCachingUSD/Source/ChaosCachingUSD/Public/ChaosCachingUSD/Operations.h b/Engine/Plugins/Experimental/ChaosCachingUSD/Source/ChaosCachingUSD/Public/ChaosCachingUSD/Operations.h
--- a/Engine/Plugins/Experimental/ChaosCachingUSD/Source/ChaosCachingUSD/Public/ChaosCachingUSD/Operations.h
+++ b/Engine/Plugins/Experimental/ChaosCachingUSD/Source/ChaosCachingUSD/Public/ChaosCachingUSD/Operations.h
@@ -150,6 +150,8 @@ namespace UE::ChaosCachingUSD
 
 	/** Get points from an attribute. Default time is used if \p Time is \c -DBL_MAX. */
 	CHAOSCACHINGUSD_API bool ReadPoints(const UE::FUsdStage& Stage, const FString& PrimPath, const FString& AttrPath, const double Time, pxr::VtArray<pxr::GfVec3f>& Points);
+	CHAOSCACHINGUSD_API bool ReadPoints(const UE::FUsdStage& Stage, const FString& PrimPath, const FString& AttrPath, const double Time, TArray<Chaos::TVector<Chaos::FRealSingle,3>>& Points);
+
 	/** Get points from the points attribute. Default time is used if \p Time is \c -DBL_MAX. */
 	CHAOSCACHINGUSD_API bool ReadPoints(const UE::FUsdStage& Stage, const FString& PrimPath, const double Time, pxr::VtArray<pxr::GfVec3f>& Points, pxr::VtArray<pxr::GfVec3f>& VtVels);
 
diff --git a/Engine/Plugins/Experimental/ChaosFlesh/Shaders/Private/DIFleshDeformer.ush b/Engine/Plugins/Experimental/ChaosFlesh/Shaders/Private/DIFleshDeformer.ush
--- a/Engine/Plugins/Experimental/ChaosFlesh/Shaders/Private/DIFleshDeformer.ush
+++ b/Engine/Plugins/Experimental/ChaosFlesh/Shaders/Private/DIFleshDeformer.ush
@@ -32,6 +32,8 @@ uint ReadNumVertices_{DataInterfaceName}(uint Index)
 	return {DataInterfaceName}_NumVertices;
 }
 
+#define ENABLE_DEFORMER_FLESH 1 // for syntax highlighting
+
 //
 // Tet mesh points, rest and dynamic 
 //
@@ -228,7 +230,7 @@ half ReadMask_{DataInterfaceName}(uint VertexIndex)
 float3x3 GetTangentBasis( float3 TangentZ )
 {
     const float Sign = TangentZ.z >= 0 ? 1 : -1;
-    const float a = -rcp( Sign + TangentZ.z );
+	const float a = -rcp( Sign + TangentZ.z );
     const float b = TangentZ.x * TangentZ.y * a;
 
     float3 TangentX = { 1 + Sign * a * Pow2( TangentZ.x ), Sign * b, -Sign * TangentZ.x };
diff --git a/Engine/Plugins/Experimental/ChaosFlesh/Source/ChaosFleshEditor/ChaosFleshEditor.Build.cs b/Engine/Plugins/Experimental/ChaosFlesh/Source/ChaosFleshEditor/ChaosFleshEditor.Build.cs
--- a/Engine/Plugins/Experimental/ChaosFlesh/Source/ChaosFleshEditor/ChaosFleshEditor.Build.cs
+++ b/Engine/Plugins/Experimental/ChaosFlesh/Source/ChaosFleshEditor/ChaosFleshEditor.Build.cs
@@ -1,11 +1,25 @@
 // Copyright Epic Games, Inc. All Rights Reserved.
 
+using System.Linq;
+
 namespace UnrealBuildTool.Rules
 {
 	public class ChaosFleshEditor : ModuleRules
 	{
         public ChaosFleshEditor(ReadOnlyTargetRules Target) : base(Target)
 		{
+			PublicIncludePaths.AddRange(
+				new string[] {
+					// ... add public include paths required here ...
+				}
+				);
+
+			PrivateIncludePaths.AddRange(
+				new string[] {
+					// ... add other private include paths required here ...
+				}
+				);
+
 			PublicDependencyModuleNames.AddRange(
 				new string[]
 				{
@@ -14,6 +28,8 @@ namespace UnrealBuildTool.Rules
 					"AssetRegistry",
 					"AssetTools",
 					"Chaos",
+					"ChaosCaching",
+					"ChaosCachingUSD",
 					"ChaosFlesh",
 					"ChaosFleshEngine",
 					"Core",
@@ -22,15 +38,19 @@ namespace UnrealBuildTool.Rules
 					"DataflowEditor",
 					"DataflowEngine",
 					"DataflowEnginePlugin",
+					"DataflowEditor",
+					"EditorStyle",
 					"Engine",
 					"EditorFramework",
 					"EditorStyle",
+					"GeometryCore",
 					"GraphEditor",
 					"InputCore",
 					"LevelEditor",
+					"MeshConversion",
 					"MeshDescription",
 					"Slate",
-				    "SlateCore",
+					"GeometryCache",
 					"Projects",
 					"PropertyEditor",
 					"RHI",
@@ -38,13 +58,94 @@ namespace UnrealBuildTool.Rules
 					"RenderCore",
 					"SceneOutliner",
 					"SkeletonEditor",
+					"Slate",
+					"SlateCore",
 					"StaticMeshDescription",
 					"ToolMenus",
 					"UnrealEd",
 				}
 			);
 
+			PrivateDependencyModuleNames.AddRange(
+				new string[] {
+				"UnrealUSDWrapper",
+				"USDClasses",
+				"USDUtilities",
+				}
+			);
+
+			DynamicallyLoadedModuleNames.AddRange(
+				new string[]
+				{
+					// ... add any modules that your module loads dynamically here ...
+				}
+				);
+
 			PrivateDefinitions.Add("CHAOS_INCLUDE_LEVEL_1=1");
+
+			// Build flavor selection copied from UnrealUSDWrapper, then modified.
+			// Currently only Win64 is supported.
+			if (EnableUsdSdk(Target) && (Target.Type == TargetType.Editor && Target.Platform == UnrealTargetPlatform.Win64))
+			{
+				bUseRTTI = true;
+				PublicDefinitions.Add("DO_USD_CACHING=1");
+			}
+			else
+			{
+				PublicDefinitions.Add("DO_USD_CACHING=0");
+			}
+		}
+		bool EnableUsdSdk(ReadOnlyTargetRules Target)
+		{
+			// USD SDK has been built against Python 3 and won't launch if the editor is using Python 2
+
+			bool bEnableUsdSdk = (
+				Target.WindowsPlatform.Compiler != WindowsCompiler.Clang &&
+				Target.StaticAnalyzer == StaticAnalyzer.None
+			);
+
+			// Don't enable USD when running the include tool because it has issues parsing Boost headers
+			if (Target.GlobalDefinitions.Contains("UE_INCLUDE_TOOL=1"))
+			{
+				bEnableUsdSdk = false;
+			}
+
+			// If you want to use USD in a monolithic target, you'll have to use the ANSI allocator.
+			// USD always uses the ANSI C allocators directly. In a DLL UE build (so not monolithic) we can just override the operators new and delete
+			// on each module with versions that use either the ANSI (so USD-compatible) allocators or the UE allocators (ModuleBoilerplate.h) when appropriate.
+			// In a monolithic build we can't do that, as the primary game module will already define overrides for operator new and delete with
+			// the standard UE allocators: Since we can only have one operator new/delete override on the entire monolithic executable, we can't define our own overrides.
+			// Additionally, the ANSI allocator does not work properly with FMallocPoisonProxy. Consequently, FMallocPoisonProxy has to be disabled.
+			// The only way around it is by forcing the ansi allocator and disabling FMallocPoisonProxy in your project's target file
+			// (YourProject/Source/YourProject.Target.cs) file like this:
+			//
+			//		public class YourProject : TargetRules
+			//		{
+			//			public YourProject(TargetInfo Target) : base(Target)
+			//			{
+			//				...
+			//				GlobalDefinitions.Add("FORCE_ANSI_ALLOCATOR=1");
+			//				GlobalDefinitions.Add("UE_USE_MALLOC_FILL_BYTES=0");
+			//				...
+			//			}
+			//		}
+			//
+			// This will force the entire built executable to use the ANSI C allocators for everything (by disabling the UE overrides in ModuleBoilerplate.h) while
+			// FMallocPoisonProxy is disabled, and so UE and USD allocations will be compatible.
+			// Note that by that point everything will be using the USD-compatible ANSI allocators anyway, so our overrides in USDMemory.h are also disabled, as they're unnecessary.
+			// Also note that we're forced to use dynamic linking for monolithic targets mainly because static linking the USD libraries disables support for user USD plugins, and secondly
+			// because those static libraries would need to be linked with the --whole-archive argument, and there is currently no standard way of doing that in UE.
+			if (bEnableUsdSdk && Target.LinkType == TargetLinkType.Monolithic && !Target.GlobalDefinitions.Contains("FORCE_ANSI_ALLOCATOR=1") && !Target.GlobalDefinitions.Contains("UE_USE_MALLOC_FILL_BYTES=0"))
+			{
+				PublicDefinitions.Add("USD_FORCE_DISABLED=1");
+				bEnableUsdSdk = false;
+			}
+			else
+			{
+				PublicDefinitions.Add("USD_FORCE_DISABLED=0");
+			}
+
+			return bEnableUsdSdk;
 		}
 	}
 }
diff --git a/Engine/Plugins/Experimental/ChaosFlesh/Source/ChaosFleshEditor/Private/ChaosFlesh/ChaosFleshEditorPlugin.cpp b/Engine/Plugins/Experimental/ChaosFlesh/Source/ChaosFleshEditor/Private/ChaosFlesh/ChaosFleshEditorPlugin.cpp
--- a/Engine/Plugins/Experimental/ChaosFlesh/Source/ChaosFleshEditor/Private/ChaosFlesh/ChaosFleshEditorPlugin.cpp
+++ b/Engine/Plugins/Experimental/ChaosFlesh/Source/ChaosFleshEditor/Private/ChaosFlesh/ChaosFleshEditorPlugin.cpp
@@ -43,6 +43,17 @@ void IChaosFleshEditorPlugin::StartupModule()
 			FConsoleCommandWithWorldAndArgsDelegate::CreateStatic(&FChaosFleshCommands::FindQualifyingTetrahedra),
 			ECVF_Default
 		));
+
+		EditorCommands.Add(IConsoleManager::Get().RegisterConsoleCommand(
+			TEXT("ChaosDeformable.CreateGeometryCache"),
+			TEXT("With an actor with flesh component(s) and a chaos cache manager selected (or use arg UsdFile), "
+				"generates a GeometryCache asset from the topology of associated SkeletalMeshComponent's import geometry, "
+				"and the simulation results from the USD file.  Requires deformer bindings for the import geometry in the "
+				"flesh component rest collection."
+				"Use arg 'UsdFile </path/to/file.usd>' to specify a specific USD file, rather than infering it from a chaos cache manager."),
+			FConsoleCommandWithWorldAndArgsDelegate::CreateStatic(&FChaosFleshCommands::CreateGeometryCache),
+			ECVF_Default
+		));
 	}
 
 	// register details customization
diff --git a/Engine/Plugins/Experimental/ChaosFlesh/Source/ChaosFleshEditor/Private/ChaosFlesh/Cmd/ChaosFleshCommands.cpp b/Engine/Plugins/Experimental/ChaosFlesh/Source/ChaosFleshEditor/Private/ChaosFlesh/Cmd/ChaosFleshCommands.cpp
--- a/Engine/Plugins/Experimental/ChaosFlesh/Source/ChaosFleshEditor/Private/ChaosFlesh/Cmd/ChaosFleshCommands.cpp
+++ b/Engine/Plugins/Experimental/ChaosFlesh/Source/ChaosFleshEditor/Private/ChaosFlesh/Cmd/ChaosFleshCommands.cpp
@@ -7,8 +7,19 @@
 #include "ChaosFlesh/FleshAsset.h"
 #include "ChaosFlesh/FleshCollection.h"
 #include "ChaosFlesh/FleshCollectionUtility.h"
-#include "ChaosFlesh/FleshComponent.h"
+#include "ChaosFlesh/ChaosDeformableTetrahedralComponent.h"
+
+#include "Chaos/CacheManagerActor.h"
+#include "Chaos/Matrix.h"
 #include "Chaos/Tetrahedron.h"
+#include "ChaosCache/FleshComponentCacheAdapter.h"
+
+#if USE_USD_SDK && DO_USD_CACHING
+#include "HAL/PlatformFile.h"
+#include "HAL/PlatformFileManager.h"
+#include "ChaosCachingUSD/Operations.h"
+#include "USDConversionUtils.h"
+#endif
 
 #include "AssetRegistry/AssetRegistryModule.h"
 #include "Editor.h"
@@ -16,12 +27,85 @@
 #include "Engine/Selection.h"
 #include "Engine/World.h"
 #include "EngineUtils.h"
+
+#include "GeometryCache.h"
+#include "GeometryCacheCodecV1.h"
+#include "GeometryCacheMeshData.h"
+#include "GeometryCacheTrackStreamable.h"
+#include "GeometryCacheConstantTopologyWriter.h"
+
+#include "GeometryCollection/Facades/CollectionTetrahedralBindingsFacade.h"
 #include "GeometryCollection/TransformCollection.h"
+
+#include "Engine/SkeletalMesh.h"
+#include "Components/SkeletalMeshComponent.h"
+#include "MeshDescription.h"
+#include "MeshDescriptionToDynamicMesh.h"
 #include "Misc/Paths.h"
+#include "ObjectTools.h"
+#include "PackageTools.h"
+#include "Rendering/SkeletalMeshLODImporterData.h"
+#include "Rendering/SkeletalMeshModel.h"
 #include "UObject/Package.h"
+#include "UObject/SavePackage.h"
+
+#include "FileHelpers.h"
+
+
+DEFINE_LOG_CATEGORY_STATIC(UChaosFleshCommandsLogging, Log, All);
+
+namespace ArgParse {
+
+	int32 ArgIndex(const TArray<FString>& Args, const FString& Target)
+	{
+		return Args.Find(Target);
+	}
+
+	bool ArgExists(const TArray<FString>& Args, const FString& Target)
+	{
+		return ArgIndex(Args, Target) != INDEX_NONE;
+	}
+
+	bool ArgStringValue(const TArray<FString>& Args, const FString& Target, FString& Value)
+	{
+		for (int32 i = 0; i < Args.Num() - 1; i++)
+		{
+			if (Args[i] == Target)
+			{
+				Value = Args[i + 1];
+				Value.RemoveFromStart(TEXT("'"));
+				Value.RemoveFromStart(TEXT("\""));
+				Value.RemoveFromEnd(TEXT("'"));
+				Value.RemoveFromEnd(TEXT("\""));
+				return true;
+			}
+		}
+		return false;
+	}
 
-DEFINE_LOG_CATEGORY_STATIC(UChaosFleshCommandsLogging, NoLogging, All);
+	bool ArgFloatValue(const TArray<FString>& Args, const FString& Target, float& Value)
+	{
+		FString StrValue;
+		if (ArgStringValue(Args, Target, StrValue))
+		{
+			Value = FCString::Atof(*StrValue);
+			return true;
+		}
+		return false;
+	}
+
+	bool ArgIntValue(const TArray<FString>& Args, const FString& Target, int32& Value)
+	{
+		FString StrValue;
+		if (ArgStringValue(Args, Target, StrValue))
+		{
+			Value = FCString::Atoi(*StrValue);
+			return true;
+		}
+		return false;
+	}
 
+} // namespace ArgParse
 
 void FChaosFleshCommands::ImportFile(const TArray<FString>& Args, UWorld* World)
 {
@@ -68,7 +152,8 @@ FChaosFleshCommands::FindQualifyingTetrahedra(const TArray<FString>& Args, UWorl
 				const TSet<UActorComponent*>& Components = Actor->GetComponents();
 				for (TSet<UActorComponent*>::TConstIterator CompIt = Components.CreateConstIterator(); CompIt; ++CompIt)
 				{
-					if (UFleshComponent* FleshComponent = Cast<UFleshComponent>(*CompIt))
+					//if (UFleshComponent* FleshComponent = Cast<UFleshComponent>(*CompIt))
+					if (UDeformableTetrahedralComponent* FleshComponent = Cast<UDeformableTetrahedralComponent >(*CompIt))
 					{
 						if (const UFleshAsset* RestCollection = FleshComponent->GetRestCollection())
 						{
@@ -99,82 +184,16 @@ FChaosFleshCommands::FindQualifyingTetrahedra(const TArray<FString>& Args, UWorl
 							float YCoordLT = -TNumericLimits<float>::Max();
 							float ZCoordLT = -TNumericLimits<float>::Max();
 							bool bHideTets = false;
-							for (int32 Idx = 0; Idx < Args.Num(); Idx++)
-							{
-								if (Args[Idx].Equals(FString(TEXT("MaxAR"))))
-								{
-									if (Args.IsValidIndex(Idx + 1))
-									{
-										Idx++;
-										MaxAR = FCString::Atof(*Args[Idx]);
-									}
-								}
-								else if (Args[Idx].Equals(FString(TEXT("MinVol"))))
-								{
-									if (Args.IsValidIndex(Idx + 1))
-									{
-										Idx++;
-										MinVol = FCString::Atof(*Args[Idx]);
-									}
-								}
 
-								// XYZ greater than
-								else if (Args[Idx].Equals(FString(TEXT("XCoordGT"))))
-								{
-									if (Args.IsValidIndex(Idx + 1))
-									{
-										Idx++;
-										XCoordGT = FCString::Atof(*Args[Idx]);
-									}
-								}
-								else if (Args[Idx].Equals(FString(TEXT("YCoordGT"))))
-								{
-									if (Args.IsValidIndex(Idx + 1))
-									{
-										Idx++;
-										YCoordGT = FCString::Atof(*Args[Idx]);
-									}
-								}
-								else if (Args[Idx].Equals(FString(TEXT("ZCoordGT"))))
-								{
-									if (Args.IsValidIndex(Idx + 1))
-									{
-										Idx++;
-										ZCoordGT = FCString::Atof(*Args[Idx]);
-									}
-								}
-
-								// XYZ less than
-								else if (Args[Idx].Equals(FString(TEXT("XCoordLT"))))
-								{
-									if (Args.IsValidIndex(Idx + 1))
-									{
-										Idx++;
-										XCoordLT = FCString::Atof(*Args[Idx]);
-									}
-								}
-								else if (Args[Idx].Equals(FString(TEXT("YCoordLT"))))
-								{
-									if (Args.IsValidIndex(Idx + 1))
-									{
-										Idx++;
-										YCoordLT = FCString::Atof(*Args[Idx]);
-									}
-								}
-								else if (Args[Idx].Equals(FString(TEXT("ZCoordLT"))))
-								{
-									if (Args.IsValidIndex(Idx + 1))
-									{
-										Idx++;
-										ZCoordLT = FCString::Atof(*Args[Idx]);
-									}
-								}
-
-								else if (Args[Idx].Equals(FString(TEXT("HideTets"))))
-								{
-									bHideTets = true;
-								}
-							}
+							ArgParse::ArgFloatValue(Args, FString(TEXT("MaxAR")), MaxAR);
+							ArgParse::ArgFloatValue(Args, FString(TEXT("MinVol")), MinVol);
+							ArgParse::ArgFloatValue(Args, FString(TEXT("XCoordGT")), XCoordGT);
+							ArgParse::ArgFloatValue(Args, FString(TEXT("YCoordGT")), YCoordGT);
+							ArgParse::ArgFloatValue(Args, FString(TEXT("ZCoordGT")), ZCoordGT);
+							ArgParse::ArgFloatValue(Args, FString(TEXT("XCoordLT")), XCoordLT);
+							ArgParse::ArgFloatValue(Args, FString(TEXT("YCoordLT")), YCoordLT);
+							ArgParse::ArgFloatValue(Args, FString(TEXT("ZCoordLT")), ZCoordLT);
+							bHideTets = ArgParse::ArgExists(Args, FString(TEXT("HideTets")));
 
 							TArray<int32> Indices;
 							for (int32 TetMeshIdx = 0; TetMeshIdx < TetrahedronStart->Num(); TetMeshIdx++)
@@ -303,4 +322,426 @@ FChaosFleshCommands::FindQualifyingTetrahedra(const TArray<FString>& Args, UWorl
 			}
 		}
 	}
-}
\ No newline at end of file
+}
+
+namespace GeometryCacheTranslatorImpl {
+
+	template<typename T>
+	T* CreateObjectInstance(UObject*& InParent, const FString& ObjectName, const EObjectFlags Flags)
+	{
+		// Parent package to place new asset
+		UPackage* Package = nullptr;
+		FString NewPackageName;
+
+		// Setup package name and create one accordingly
+		NewPackageName = FPackageName::GetLongPackagePath(InParent->GetOutermost()->GetPathName()) + TEXT("/") + ObjectName;
+		NewPackageName = UPackageTools::SanitizePackageName(NewPackageName);
+		Package = CreatePackage(*NewPackageName);
+
+		const FString SanitizedObjectName = ObjectTools::SanitizeObjectName(ObjectName);
+
+		T* ExistingTypedObject = FindObject<T>(Package, *SanitizedObjectName);
+		UObject* ExistingObject = FindObject<UObject>(Package, *SanitizedObjectName);
+
+		if (ExistingTypedObject != nullptr)
+		{
+			ExistingTypedObject->PreEditChange(nullptr);
+		}
+		else if (ExistingObject != nullptr)
+		{
+			// Replacing an object.  Here we go!
+			// Delete the existing object
+			const bool bDeleteSucceeded = ObjectTools::DeleteSingleObject(ExistingObject);
+
+			if (bDeleteSucceeded)
+			{
+				// Force GC so we can cleanly create a new asset (and not do an 'in place' replacement)
+				CollectGarbage(GARBAGE_COLLECTION_KEEPFLAGS);
+
+				// Create a package for each mesh
+				Package = CreatePackage(*NewPackageName);
+				InParent = Package;
+			}
+			else
+			{
+				// failed to delete
+				return nullptr;
+			}
+		}
+
+		return NewObject<T>(Package, FName(*SanitizedObjectName), Flags | RF_Public);
+	}
+
+	TOptional<TArray<int32>>
+	GetMeshImportVertexMap(const USkinnedAsset& SkeletalMeshAsset)
+	{
+		constexpr int32 LODIndex = 0;
+		const TOptional<TArray<int32>> None;
+
+		const FSkeletalMeshModel* const MLDModel = SkeletalMeshAsset.GetImportedModel();
+		if (!MLDModel || !MLDModel->LODModels.IsValidIndex(LODIndex))
+		{
+			return None;
+		}
+
+		const FSkeletalMeshLODModel& MLDLOD = MLDModel->LODModels[LODIndex];
+		const TArray<int32>& Map = MLDLOD.MeshToImportVertexMap;
+		if (Map.IsEmpty())
+		{
+			UE_LOG(UChaosFleshCommandsLogging, Warning, 
+				TEXT("MeshToImportVertexMap is empty. MLDeformer Asset should be an imported SkeletalMesh (e.g. from fbx)."));
+			return None;
+		}
+
+		return Map;
+	}
+
+} //namespace GeometryCacheTranslatorImpl
+
+void FChaosFleshCommands::CreateGeometryCache(const TArray<FString>& Args, UWorld* World)
+{
+	FPlatformFileManager& FileManager = FPlatformFileManager::Get();
+	IPlatformFile& PlatformFile = FileManager.GetPlatformFile();
+
+	// Get the current selection.
+	USelection* SelectedActors = GEditor->GetSelectedActors();
+	if (!SelectedActors)
+	{
+		UE_LOG(UChaosFleshCommandsLogging, Error,
+			TEXT("CreateGeometryCache - No ChaosCacheManager or Actor with FleshComponent(s) found in selection."));
+		return;
+	}
+
+	// Find cache file, if specified.
+	FString UsdFileOverride;
+	bool bUsdFile = ArgParse::ArgStringValue(Args, FString(TEXT("UsdFile")), UsdFileOverride);
+
+	// Find a ChaosCacheManager, if a cache file hasn't been specified.
+	AChaosCacheManager* CacheManager = nullptr;
+	if (bUsdFile)
+	{
+		if (!PlatformFile.FileExists(*UsdFileOverride))
+		{
+			UE_LOG(UChaosFleshCommandsLogging, Error,
+				TEXT("CreateGeometryCache - File not found: '%s'"), *UsdFileOverride);
+			return;
+		}
+	}
+	else
+	{
+		for (FSelectionIterator ActorIt(*SelectedActors); ActorIt; ++ActorIt)
+		{
+			if (AActor* Actor = Cast<AActor>(*ActorIt))
+			{
+				CacheManager = Cast<AChaosCacheManager>(Actor);
+				if (CacheManager)
+				{
+					break;
+				}
+			}
+		}
+		if (!CacheManager)
+		{
+			UE_LOG(UChaosFleshCommandsLogging, Error,
+				TEXT("CreateGeometryCache - No ChaosCacheManager found in selection, and no cache file specified."));
+			return;
+		}
+	}
+	const TArray<FObservedComponent>& ObservedComponents = 
+		CacheManager ? CacheManager->GetObservedComponents() : TArray<FObservedComponent>();
+
+	// Find actors with flesh components.
+	TArray<UPackage*> PackagesToSave;
+	TArray<AActor*> Actors;
+	for (FSelectionIterator ActorIt(*SelectedActors); ActorIt; ++ActorIt)
+	{
+		AActor* Actor = Cast<AActor>(*ActorIt);
+		if (!Actor || Actor == CacheManager || Actors.Contains(Actor))
+		{
+			continue;
+		}
+		Actors.Add(Actor);
+
+		//
+		// Find flesh component w/cache file, and skeletal mesh components
+		//
+
+		TArray<UDeformableTetrahedralComponent*> FleshComponents;
+		TArray<FString> USDCacheFilePaths;
+		TArray<const USkeletalMesh*> SkeletalMeshComponents;
+
+		const TSet<UActorComponent*>& Components = Actor->GetComponents();
+		for (TSet<UActorComponent*>::TConstIterator CompIt = Components.CreateConstIterator(); CompIt; ++CompIt)
+		{
+			if (UDeformableTetrahedralComponent* FleshComponent = Cast<UDeformableTetrahedralComponent>(*CompIt))
+			{
+				if (FleshComponents.Contains(FleshComponent))
+				{
+					continue;
+				}
+
+#if USE_USD_SDK && DO_USD_CACHING
+				if (bUsdFile)
+				{
+					FleshComponents.Add(FleshComponent);
+					USDCacheFilePaths.Add(UsdFileOverride);
+				}
+				else
+				{
+					for (const FObservedComponent& ObservedComponent : ObservedComponents)
+					{
+						FString CacheFilePath = Chaos::FFleshCacheAdapter::GetUSDCacheFilePathRO(ObservedComponent, FleshComponent);
+						if (PlatformFile.FileExists(*CacheFilePath))
+						{
+							FleshComponents.Add(FleshComponent);
+							USDCacheFilePaths.Add(CacheFilePath);
+						}
+						else
+						{
+							UE_LOG(UChaosFleshCommandsLogging, Warning,
+								TEXT("CreateGeometryCache - Failed to find USD file: '%s'"),
+								*CacheFilePath);
+						}
+					}
+				}
+#endif // USE_USD_SDK && DO_USD_CACHING
+			}
+			else if (const USkeletalMeshComponent* SkeletalMeshComponent = Cast<const USkeletalMeshComponent>(*CompIt))
+			{
+				if (const USkeletalMesh* SkeletalMesh = SkeletalMeshComponent->GetSkeletalMeshAsset())
+				{
+					SkeletalMeshComponents.Add(SkeletalMesh);
+				}
+			}
+		}
+		if (FleshComponents.IsEmpty() || SkeletalMeshComponents.IsEmpty())
+		{
+			UE_LOG(UChaosFleshCommandsLogging, Warning,
+				TEXT("CreateGeometryCache - Failed to find any FleshComponents with a valid simulation cache file and a skeletal mesh component for actor: '%s'"),
+				*Actor->GetName());
+			continue;
+		}
+
+		//
+		// Create GeometryCache instance, named after the actor
+		// 
+
+		UObject* InParent = Actor;
+		EObjectFlags Flags = RF_Public | RF_Standalone; //  RF_NoFlags;
+
+		FString GCName = FPaths::GetBaseFilename(InParent->GetName());
+		GCName = FPaths::SetExtension(GCName, FString(TEXT("uasset")));
+
+		UGeometryCache* GeometryCache =
+			GeometryCacheTranslatorImpl::CreateObjectInstance<UGeometryCache>(
+				InParent,
+				GCName,
+				Flags);
+		if (!GeometryCache)
+		{
+				UE_LOG(UChaosFleshCommandsLogging, Error, 
+					TEXT("CreateGeometryCache - Failed to create geometry cache instance for actor: '%s'"),
+					*InParent->GetName());
+				continue;
+		}
+		else
+		{
+			UE_LOG(UChaosFleshCommandsLogging, Log,
+				TEXT("CreateGeometryCache - Created geometry cache instance: '%s'"),
+				*GeometryCache->GetName());
+		}
+
+		for (int32 i = 0; i < FleshComponents.Num(); i++)
+		{
+			UDeformableTetrahedralComponent* FleshComponent = FleshComponents[i];
+			const FString& USDCacheFilePath = USDCacheFilePaths[i];
+
+			const UFleshAsset* RestCollection = FleshComponent->GetRestCollection();
+			if (!RestCollection)
+			{
+				continue;
+			}
+			const FFleshCollection* FleshCollection = RestCollection->GetCollection();
+			if (!FleshCollection)
+			{
+				continue;
+			}
+
+			const TManagedArray<FVector3f>* RestVertices =
+				RestCollection->FindPositions();
+
+			const TManagedArray<FIntVector4>* TetMesh =
+				FleshCollection->FindAttribute<FIntVector4>(
+					FTetrahedralCollection::TetrahedronAttribute, FTetrahedralCollection::TetrahedralGroup);
+			const TManagedArray<int32>* TetrahedronStart =
+				FleshCollection->FindAttribute<int32>(
+					FTetrahedralCollection::TetrahedronStartAttribute, FGeometryCollection::GeometryGroup);
+			const TManagedArray<int32>* TetrahedronCount =
+				FleshCollection->FindAttribute<int32>(
+					FTetrahedralCollection::TetrahedronCountAttribute, FGeometryCollection::GeometryGroup);
+
+			GeometryCollection::Facades::FTetrahedralBindings TetBindings(*FleshCollection);
+
+			const int32 LODIndex = 0;
+			for (const USkeletalMesh* SkeletalMesh : SkeletalMeshComponents)
+			{
+				TOptional<TArray<int32>> OptionalMap = GeometryCacheTranslatorImpl::GetMeshImportVertexMap(*SkeletalMesh);
+				const TArray<int32>* MeshToImportVertexMap = OptionalMap.IsSet() ? &OptionalMap.GetValue() : nullptr;
+
+				//
+				// Extract bindings from tet mesh to skel mesh import geometry in the rest collection.
+				//
+
+				FString MeshId = GeometryCollection::Facades::FTetrahedralBindings::GetMeshId(SkeletalMesh, false);
+				FName MeshIdName(MeshId);
+				const int32 TetIndex = TetBindings.GetTetMeshIndex(MeshIdName, LODIndex);
+				if (TetIndex == INDEX_NONE)
+				{
+					UE_LOG(UChaosFleshCommandsLogging, Error,
+						TEXT("CreateGeometryCache - No tet mesh index associated with mesh '%s' LOD: %d"),
+						*MeshId, LODIndex);
+					continue;
+				}
+				if (!TetBindings.ReadBindingsGroup(TetIndex, MeshIdName, LODIndex))
+				{
+					UE_LOG(UChaosFleshCommandsLogging, Error,
+						TEXT("CreateGeometryCache - Failed to read bindings group associated with mesh '%s' LOD: %d"),
+						*MeshId, LODIndex);
+					continue;
+				}
+
+				TUniquePtr<GeometryCollection::Facades::FTetrahedralBindings::Evaluator> BindingsEvalPtr = 
+					TetBindings.InitEvaluator(RestVertices);
+				const GeometryCollection::Facades::FTetrahedralBindings::Evaluator& BindingsEval = *BindingsEvalPtr.Get();
+				if (!BindingsEval.IsValid())
+				{
+					UE_LOG(UChaosFleshCommandsLogging, Error,
+						TEXT("CreateGeometryCache - Bindings group associated with mesh '%s' LOD: %d has invalid data."),
+						*MeshId, LODIndex);
+					continue;
+				}
+				const int32 NumVertices = BindingsEval.NumVertices();
+
+				TArray<TArray<FVector3f>> FramePositions;
+#if USE_USD_SDK && DO_USD_CACHING
+				//
+				// Open the USD cache, get all time samples.
+				//
+
+				UE::FUsdStage USDStage;
+				if (!UE::ChaosCachingUSD::OpenStage(USDCacheFilePath, USDStage))
+				{
+					// Warn
+					continue;
+				}
+				else
+				{
+					UE_LOG(UChaosFleshCommandsLogging, Log,
+						TEXT("CreateGeometryCache - Opened USD stage: '%s'"),
+						*USDCacheFilePath);
+				}
+
+				FString PrimPath = UsdUtils::GetPrimPathForObject(FleshComponent);
+				TArray<double> TimeSamples;
+				if (!UE::ChaosCachingUSD::ReadTimeSamples(USDStage, PrimPath, TimeSamples))
+				{
+					// Warn
+					continue;
+				}
+
+				//
+				// Deform render geometry, storing per frame data in CurrFramePositions.
+				//
+
+				UE_LOG(UChaosFleshCommandsLogging, Log,
+					TEXT("CreateGeometryCache - FleshComponent '%s' deforming SkeleltalMesh '%s' render geometry over %d time samples..."),
+					*FleshComponent->GetName(), *SkeletalMesh->GetName(), TimeSamples.Num());
+
+				FramePositions.SetNum(TimeSamples.Num());
+
+				TArray<Chaos::TVector<Chaos::FRealSingle, 3>> CurrTetVertices;
+				int32 CurrTimeIndex = 0;
+				for (const double Time : TimeSamples)
+				{
+					if (!UE::ChaosCachingUSD::ReadPoints(
+						USDStage, PrimPath, UE::ChaosCachingUSD::GetPointsAttrName(), Time, CurrTetVertices))
+					{
+						// Warn
+						continue;
+					}
+
+					TArray<FVector3f>& CurrFramePositions = FramePositions[CurrTimeIndex++];
+					CurrFramePositions.SetNumUninitialized(NumVertices);
+					for (int j = 0; j < NumVertices; j++)
+					{
+						CurrFramePositions[j] = BindingsEval.GetEmbeddedPosition(j, CurrTetVertices);
+					}
+				}
+
+				UE::ChaosCachingUSD::CloseStage(USDStage);
+#else // USE_USD_SDK && DO_USD_CACHING
+				// TODO: Get points from chaos cache?
+				UE_LOG(UChaosFleshCommandsLogging, Error, TEXT("USD Caching is not supported on this platform."));
+				return;
+#endif // USE_USD_SDK && DO_USD_CACHING
+
+				//
+				// Write deformed render vertices to GeometryCache.
+				//
+
+				UE_LOG(UChaosFleshCommandsLogging, Log,
+					TEXT("CreateGeometryCache - FleshComponent '%s' writing deformed SkeleltalMesh '%s' render geometry to geometry cache: '%s'"),
+					*FleshComponent->GetName(), *SkeletalMesh->GetName(), *GeometryCache->GetName());
+
+				UE::GeometryCacheHelpers::FGeometryCacheConstantTopologyWriter Writer(*GeometryCache);
+
+				// Writes Indices, UV's, Colors, ImportedVertexNumbers, and "BatchesInfo".
+				const int32 Index = UE::GeometryCacheHelpers::AddTrackWriterFromSkinnedAsset(Writer, *SkeletalMesh);
+				if (Index == INDEX_NONE)
+				{
+					UE_LOG(UChaosFleshCommandsLogging, Error,
+						TEXT("CreateGeometryCache - FleshComponent '%s' failed to write topology to geometry cache."),
+						*FleshComponent->GetName());
+					continue;
+				}
+
+				UE::GeometryCacheHelpers::FGeometryCacheConstantTopologyWriter::FTrackWriter& TrackWriter = Writer.GetTrackWriter(Index);
+				if (MeshToImportVertexMap)
+				{
+					TrackWriter.ImportedVertexNumbers.SetNum(MeshToImportVertexMap->Num());
+					for (int32 j = 0; j < MeshToImportVertexMap->Num(); j++)
+					{
+						TrackWriter.ImportedVertexNumbers[j] = (*MeshToImportVertexMap)[j];
+					}
+				}
+
+				TArrayView<TArray<FVector3f>> FramePositionsView(FramePositions);
+				if (!TrackWriter.WriteAndClose(FramePositionsView))
+				{
+					UE_LOG(UChaosFleshCommandsLogging, Error,
+						TEXT("CreateGeometryCache - FleshComponent '%s' failed to write vertices track to geometry cache."),
+						*FleshComponent->GetName());
+					continue;
+				}
+
+				UE_LOG(UChaosFleshCommandsLogging, Log,
+					TEXT("CreateGeometryCache - FleshComponent '%s' wrote %d frames of %d vertices to geometry cache: '%s'"),
+					*FleshComponent->GetName(), FramePositions.Num(), NumVertices,
+					*GeometryCache->GetName());
+				PackagesToSave.AddUnique(GeometryCache->GetOutermost());
+
+			} // end for USkeletalMesh
+		} // end for FleshComponents
+	} // end for SelectedActors
+
+	if (!PackagesToSave.IsEmpty())
+	{
+		// ChaosClothGenerator::SavePackage()
+		constexpr bool bCheckDirty = false;
+		constexpr bool bPromptToSave = true;
+		FEditorFileUtils::PromptForCheckoutAndSave(
+			PackagesToSave, bCheckDirty, bPromptToSave, 
+			FText::FromString(FString(TEXT("Save GeometryCache"))), 
+			FText::FromString(FString(TEXT("Save new GeometryCache assets."))));
+	}
+}
diff --git a/Engine/Plugins/Experimental/ChaosFlesh/Source/ChaosFleshEditor/Public/ChaosFlesh/Cmd/ChaosFleshCommands.h b/Engine/Plugins/Experimental/ChaosFlesh/Source/ChaosFleshEditor/Public/ChaosFlesh/Cmd/ChaosFleshCommands.h
--- a/Engine/Plugins/Experimental/ChaosFlesh/Source/ChaosFleshEditor/Public/ChaosFlesh/Cmd/ChaosFleshCommands.h
+++ b/Engine/Plugins/Experimental/ChaosFlesh/Source/ChaosFleshEditor/Public/ChaosFlesh/Cmd/ChaosFleshCommands.h
@@ -28,4 +28,18 @@ public:
 	* @param World
 	*/
 	static void FindQualifyingTetrahedra(const TArray<FString>& Args, UWorld* World);
+
+	/**
+	* Create GeometryCache asset from a cached flesh simulation.
+	* 
+	* Command invoked from "FChaosDeformableCommands.CreateGeometryCache", uses the selected 
+	* Actor's FleshComponent and SkeletalMeshComponent, optionally along with a Chaos Cache Manager.
+	* Deforms the SkeletalMesh's import geometry by the cached simulation results, and stores that
+	* surface in a new GeometryCache asset.  The rest collection asset on the flesh component(s)
+	* needs to contain deformer bindings for the import geometry.
+	* @param Args - Supported arguments:
+	*   'UsdFile </path/to/file.usd>' overrides which usd file to use, rather than inferring it from a chaos cache manager.
+	*/
+	static void CreateGeometryCache(const TArray<FString>& Args, UWorld* World);
+
 };
\ No newline at end of file
diff --git a/Engine/Plugins/Experimental/ChaosFlesh/Source/ChaosFleshEngine/Private/ChaosCache/FleshComponentCacheAdapter.cpp b/Engine/Plugins/Experimental/ChaosFlesh/Source/ChaosFleshEngine/Private/ChaosCache/FleshComponentCacheAdapter.cpp
--- a/Engine/Plugins/Experimental/ChaosFlesh/Source/ChaosFleshEngine/Private/ChaosCache/FleshComponentCacheAdapter.cpp
+++ b/Engine/Plugins/Experimental/ChaosFlesh/Source/ChaosFleshEngine/Private/ChaosCache/FleshComponentCacheAdapter.cpp
@@ -472,7 +472,8 @@ namespace Chaos
 		}
 	}
 
-	FString GetCacheDirectory(const FObservedComponent& InObserved)
+#if USE_USD_SDK && DO_USD_CACHING
+	FString FFleshCacheAdapter::GetUSDCacheDirectory(const FObservedComponent& InObserved)
 	{
 		// USDCacheDirectory is relative to the content dir, but with "/Game" rather than just "/" or some relative path.
 		FString CacheDir = InObserved.USDCacheDirectory.Path;
@@ -488,7 +489,7 @@ namespace Chaos
 		return CacheDir;
 	}
 
-	FString GetCacheFileName(const UFleshComponent* FleshComp)
+	FString FFleshCacheAdapter::GetUSDCacheFileName(const UDeformablePhysicsComponent* FleshComp)
 	{
 		const UObject* CurrObject = FleshComp;
 		const AActor* Actor = nullptr;
@@ -514,6 +515,49 @@ namespace Chaos
 		return CompName;
 	}
 
+	FString FFleshCacheAdapter::GetUSDCacheFilePathRO(
+		const FObservedComponent& InObserved, 
+		const UDeformablePhysicsComponent* FleshComp)
+	{
+		FString CacheDir = GetUSDCacheDirectory(InObserved);
+		FString CompName = GetUSDCacheFileName(FleshComp);
+		FString PrimPath = UsdUtils::GetPrimPathForObject(FleshComp);
+
+		// Look for an existing file of a supported format.
+		TArray<FString> SupportedFormats = UnrealUSDWrapper::GetNativeFileFormats();
+		for (const FString& Ext : SupportedFormats)
+		{
+			FString FileName = FString::Printf(TEXT("%s.%s"), *CompName, *Ext);
+			FString FilePath = FPaths::Combine(CacheDir, FileName);
+			if (FPaths::FileExists(FilePath))
+			{
+				return FilePath;
+			}
+		}
+
+		// Fall back on the format we're currently set to write.
+		FString Ext = CVarParams.bWriteBinary ? FString(TEXT("usd")) : FString(TEXT("usda"));
+		FString FileName = FString::Printf(TEXT("%s.%s"), *CompName, *Ext);
+		FString FilePath = FPaths::Combine(CacheDir, FileName);
+
+		return FilePath;
+	}
+
+	FString FFleshCacheAdapter::GetUSDCacheFilePathRW(
+		const FObservedComponent& InObserved,
+		const UDeformablePhysicsComponent* FleshComp)
+	{
+		FString CacheDir = GetUSDCacheDirectory(InObserved);
+		FString CompName = GetUSDCacheFileName(FleshComp);
+		FString PrimPath = UsdUtils::GetPrimPathForObject(FleshComp);
+		FString Ext = CVarParams.bWriteBinary ? FString(TEXT("usd")) : FString(TEXT("usda"));
+		FString FileName = FString::Printf(TEXT("%s.%s"), *CompName, *Ext);
+		FString FilePath = FPaths::Combine(CacheDir, FileName);
+
+		return FilePath;
+	}
+#endif // USE_USD_SDK && DO_USD_CACHING
+
 	bool FFleshCacheAdapter::InitializeForRecord(UPrimitiveComponent* InComponent, FObservedComponent& InObserved)
 	{
 		if( FDeformableSolver* Solver = GetDeformableSolver(InComponent))
@@ -529,8 +573,8 @@ namespace Chaos
 				// USD caching
 				//
 
-				FString CompName = GetCacheFileName(FleshComp);
-				FString CacheDir = GetCacheDirectory(InObserved);
+				FString CompName = GetUSDCacheFileName(FleshComp);
+				FString CacheDir = GetUSDCacheDirectory(InObserved);
 				FPlatformFileManager& FileManager = FPlatformFileManager::Get();
 				IPlatformFile& PlatformFile = FileManager.GetPlatformFile();
 				if (!PlatformFile.DirectoryExists(*CacheDir))
@@ -617,7 +661,7 @@ namespace Chaos
 		{
 			FDeformableSolver::FGameThreadAccess GameThreadAccess(Solver, Softs::FGameThreadAccessor());
 			GameThreadAccess.SetEnableSolver(false);
-			if (UFleshComponent* FleshComp = CastChecked<UFleshComponent>(InComponent))
+			if (UDeformableTetrahedralComponent* FleshComp = CastChecked<UDeformableTetrahedralComponent>(InComponent))
 			{
 				FleshComp->ResetDynamicCollection();
 
@@ -626,15 +670,11 @@ namespace Chaos
 				// USD caching
 				//
 
-				FString CacheDir = GetCacheDirectory(InObserved);
-				FString CompName = GetCacheFileName(FleshComp);
 				PrimPath = UsdUtils::GetPrimPathForObject(FleshComp);
 				if (bUseMonolith)
 				{
 					bReadOnly = true;
-					FString Ext = CVarParams.bWriteBinary ? FString(TEXT("usd")) : FString(TEXT("usda"));
-					FString FileName = FString::Printf(TEXT("%s.%s"), *CompName, *Ext);
-					FilePath = FPaths::Combine(CacheDir, FileName);
+					FilePath = GetUSDCacheFilePathRO(InObserved, FleshComp);
 
 					FPlatformFileManager& FileManager = FPlatformFileManager::Get();
 					IPlatformFile& PlatformFile = FileManager.GetPlatformFile();
diff --git a/Engine/Plugins/Experimental/ChaosFlesh/Source/ChaosFleshEngine/Public/ChaosCache/FleshComponentCacheAdapter.h b/Engine/Plugins/Experimental/ChaosFlesh/Source/ChaosFleshEngine/Public/ChaosCache/FleshComponentCacheAdapter.h
--- a/Engine/Plugins/Experimental/ChaosFlesh/Source/ChaosFleshEngine/Public/ChaosCache/FleshComponentCacheAdapter.h
+++ b/Engine/Plugins/Experimental/ChaosFlesh/Source/ChaosFleshEngine/Public/ChaosCache/FleshComponentCacheAdapter.h
@@ -25,6 +25,7 @@ namespace Chaos::Softs
 	class FPBDEvolution;
 	class FSolverParticles;
 }
+class UDeformablePhysicsComponent;
 
 namespace Chaos
 {
@@ -46,6 +47,13 @@ namespace Chaos
 		typedef Chaos::Softs::FPBDEvolution FEvolution;
 		typedef Chaos::Softs::FSolverParticles FParticles;
 
+#if USE_USD_SDK && DO_USD_CACHING
+		static CHAOSFLESHENGINE_API FString GetUSDCacheDirectory(const FObservedComponent& InObserved);
+		static CHAOSFLESHENGINE_API FString GetUSDCacheFileName(const UDeformablePhysicsComponent* FleshComp);
+		static CHAOSFLESHENGINE_API FString GetUSDCacheFilePathRO(const FObservedComponent& InObserved, const UDeformablePhysicsComponent* FleshComp);
+		static CHAOSFLESHENGINE_API FString GetUSDCacheFilePathRW(const FObservedComponent& InObserved, const UDeformablePhysicsComponent* FleshComp);
+#endif // USE_USD_SDK && DO_USD_CACHING
+
 		virtual ~FFleshCacheAdapter();
 
 		// ~Begin FComponentCacheAdapter interface
diff --git a/Engine/Plugins/Experimental/ChaosFlesh/Source/ChaosFleshNodes/Private/Dataflow/ChaosFleshBindingsNodes.cpp b/Engine/Plugins/Experimental/ChaosFlesh/Source/ChaosFleshNodes/Private/Dataflow/ChaosFleshBindingsNodes.cpp
--- a/Engine/Plugins/Experimental/ChaosFlesh/Source/ChaosFleshNodes/Private/Dataflow/ChaosFleshBindingsNodes.cpp
+++ b/Engine/Plugins/Experimental/ChaosFlesh/Source/ChaosFleshNodes/Private/Dataflow/ChaosFleshBindingsNodes.cpp
@@ -1,6 +1,8 @@
 // Copyright Epic Games, Inc. All Rights Reserved.
 
 #include "Dataflow/ChaosFleshBindingsNodes.h"
+#include "Dataflow/DataflowInputOutput.h"
+#include "Dataflow/DataflowNodeFactory.h"
 
 #include "Chaos/AABBTree.h"
 #include "Chaos/BoundingVolumeHierarchy.h"
@@ -9,10 +11,13 @@
 #include "ChaosFlesh/TetrahedralCollection.h"
 #include "Containers/Map.h"
 #include "Engine/StaticMesh.h"
-#include "Dataflow/DataflowInputOutput.h"
 #include "Engine/SkeletalMesh.h"
-#include "Dataflow/DataflowNodeFactory.h"
 #include "GeometryCollection/Facades/CollectionTetrahedralBindingsFacade.h"
+#include "IndexTypes.h"
+#include "MeshDescription.h"
+#include "MeshDescriptionToDynamicMesh.h"
+#include "Rendering/SkeletalMeshLODModel.h"
+#include "Rendering/SkeletalMeshModel.h"
 #include "Rendering/SkeletalMeshRenderData.h"
 #include "UObject/PrimaryAssetId.h"
 
@@ -27,6 +32,27 @@ namespace Dataflow
 	}
 }
 
+void
+BuildVertexToVertexAdjacencyBuffer(
+	const UE::Geometry::FDynamicMesh3 DynamicMesh,
+	TArray<TArray<uint32>>& NeighborNodes)
+{
+	NeighborNodes.SetNum(DynamicMesh.VertexCount());
+	for (int i = 0; i < DynamicMesh.TriangleCount(); i++)
+	{
+		const UE::Geometry::FIndex3i& Tri = DynamicMesh.GetTriangle(i);
+		TArray<uint32>& N0 = NeighborNodes[Tri[0]];
+		N0.AddUnique(Tri[1]);
+		N0.AddUnique(Tri[2]);
+		TArray<uint32>& N1 = NeighborNodes[Tri[1]];
+		N1.AddUnique(Tri[0]);
+		N1.AddUnique(Tri[2]);
+		TArray<uint32>& N2 = NeighborNodes[Tri[2]];
+		N2.AddUnique(Tri[0]);
+		N2.AddUnique(Tri[1]);
+	}
+}
+
 void
 BuildVertexToVertexAdjacencyBuffer(
 	const FSkeletalMeshLODRenderData& LodRenderData,
@@ -154,50 +180,85 @@ FGenerateSurfaceBindings::Evaluate(Dataflow::FContext& Context, const FDataflowO
 			if (SkeletalMesh)
 			{
 				FPrimaryAssetId Id = SkeletalMesh->GetPrimaryAssetId();
-				if (Id.IsValid())
-				{
-					MeshId = Id.ToString();
-				}
-				else
+				MeshId = GeometryCollection::Facades::FTetrahedralBindings::GetMeshId(SkeletalMesh, bUseSkeletalMeshImportModel);
+
+				if (!bUseSkeletalMeshImportModel)
 				{
-					MeshId = SkeletalMesh->GetName();
-				}
+					FSkeletalMeshRenderData* RenderData = SkeletalMesh->GetResourceForRendering();
+
+					MeshVertices.SetNum(RenderData->LODRenderData.Num());
+					MeshNeighborNodes.SetNum(RenderData->LODRenderData.Num());
+					for (int32 i = 0; i < RenderData->LODRenderData.Num(); i++)
+					{
+						FSkeletalMeshLODRenderData* LODRenderData = &RenderData->LODRenderData[i];
+						const FPositionVertexBuffer& PositionVertexBuffer =
+							LODRenderData->StaticVertexBuffers.PositionVertexBuffer;
 
-				FSkeletalMeshRenderData* RenderData = SkeletalMesh->GetResourceForRendering();
+						TArray<FVector3f>& Vertices = MeshVertices[i];
+						Vertices.SetNumUninitialized(PositionVertexBuffer.GetNumVertices());
+						for (uint32 j = 0; j < PositionVertexBuffer.GetNumVertices(); j++)
+						{
+							const FVector3f& Pos = PositionVertexBuffer.VertexPosition(j);
+							Vertices[j] = Pos;
+						}
 
-				MeshVertices.SetNum(RenderData->LODRenderData.Num());
-				MeshNeighborNodes.SetNum(RenderData->LODRenderData.Num());
-				for (int32 i = 0; i < RenderData->LODRenderData.Num(); i++)
+						TArray<TArray<uint32>>& NeighborNodes = MeshNeighborNodes[i];
+						BuildVertexToVertexAdjacencyBuffer(*LODRenderData, NeighborNodes);
+					}
+				}
+#if WITH_EDITOR
+				else // Import Model
 				{
-					FSkeletalMeshLODRenderData* LODRenderData = &RenderData->LODRenderData[i];
-					const FPositionVertexBuffer& PositionVertexBuffer =
-						LODRenderData->StaticVertexBuffers.PositionVertexBuffer;
+					const int32 LODIndex = 0;
+					MeshVertices.SetNum(1);
+					MeshNeighborNodes.SetNum(1);
+
+					// Check first if we have bulk data available and non-empty.
+					FMeshDescription SourceMesh;
+#if WITH_EDITORONLY_DATA
+					if (SkeletalMesh->HasMeshDescription(LODIndex))
+					{
+						// @todo(brice) : Confirm correct in release 5.4
+						//FSkeletalMeshImportData SkeletalMeshImportData;
+						//SkeletalMesh->LoadLODImportedData(LODIndex, SkeletalMeshImportData);
+						//FSkeletalMeshBuildSettings Settings;
+						//SkeletalMeshImportData.GetMeshDescription(SkeletalMesh, &Settings, SourceMesh);
+						SkeletalMesh->CloneMeshDescription(LODIndex, SourceMesh);
+					}
+					else
+#endif
+					{
+						// Fall back on the LOD model directly if no bulk data exists. When we commit
+						// the mesh description, we override using the bulk data. This can happen for older
+						// skeletal meshes, from UE 4.24 and earlier.
+						const FSkeletalMeshModel* SkeletalMeshModel = SkeletalMesh->GetImportedModel();
+						if (SkeletalMeshModel && SkeletalMeshModel->LODModels.IsValidIndex(LODIndex))
+						{
+							SkeletalMeshModel->LODModels[LODIndex].GetMeshDescription(SkeletalMesh, LODIndex, SourceMesh);
+						}
+					}
 
-					TArray<FVector3f>& Vertices = MeshVertices[i];
-					Vertices.SetNumUninitialized(PositionVertexBuffer.GetNumVertices());
-					for (uint32 j = 0; j < PositionVertexBuffer.GetNumVertices(); j++)
+					UE::Geometry::FDynamicMesh3 DynamicMesh;
+					FMeshDescriptionToDynamicMesh Converter;
+					Converter.Convert(&SourceMesh, DynamicMesh);
+
+					//const FVertexArray& SourceVertices = SourceMesh.Vertices();
+					TArray<FVector3f>& Vertices = MeshVertices[LODIndex];
+					Vertices.SetNumUninitialized(DynamicMesh.VertexCount());
+					for (int32 j = 0; j < DynamicMesh.VertexCount(); j++)
 					{
-						const FVector3f& Pos = PositionVertexBuffer.VertexPosition(j);
-						Vertices[j] = Pos;
+						const FVector3d& Pos = DynamicMesh.GetVertex(j);
+						Vertices[j].Set(Pos[0], Pos[1], Pos[2]);
 					}
 
-					TArray<TArray<uint32>>& NeighborNodes = MeshNeighborNodes[i];
-					BuildVertexToVertexAdjacencyBuffer(*LODRenderData, NeighborNodes);
+					TArray<TArray<uint32>>& NeighborNodes = MeshNeighborNodes[LODIndex];
+					BuildVertexToVertexAdjacencyBuffer(DynamicMesh, NeighborNodes);
 				}
+#endif
 			}
 			else // StaticMesh
 			{
-				//StaticMesh->GetMeshId(MeshId); // not available at runtime!
-				//MeshId = RenderData->DerivedDataKey; // same problem
-				FPrimaryAssetId Id = StaticMesh->GetPrimaryAssetId();
-				if (Id.IsValid())
-				{
-					MeshId = Id.ToString();
-				}
-				else
-				{
-					MeshId = StaticMesh->GetName();
-				}
+				MeshId = GeometryCollection::Facades::FTetrahedralBindings::GetMeshId(StaticMesh);
 
 				const FStaticMeshRenderData* RenderData = StaticMesh->GetRenderData();
 				const int32 NumLOD = RenderData->LODResources.Num();
@@ -532,14 +593,16 @@ FGenerateSurfaceBindings::Evaluate(Dataflow::FContext& Context, const FDataflowO
 
 					//ELogVerbosity::Type Verbosity = Orphans.Num() > 0 ? ELogVerbosity::Error : ELogVerbosity::Display;
 					UE_LOG(LogMeshBindings, Display,
-						TEXT("'%s' - Generated mesh bindings between tet mesh index %d and render mesh of '%s' LOD %d - stats:\n"
+						TEXT("'%s' - Generated mesh bindings between tet mesh index %d and %s mesh of '%s' LOD %d - stats:\n"
 							"    Render vertices num: %d\n"
 							"    Vertices in tetrahedra: %d\n"
 							"    Vertices bound to tet surface: %d\n"
 							"    Orphaned vertices reparented: %d\n"
 							"    Vertices orphaned: %d"),
 						*GetName().ToString(),
-						TetMeshIdx, *MeshId, LOD,
+						TetMeshIdx, 
+						bUseSkeletalMeshImportModel ? TEXT("import") : TEXT("render"),
+						*MeshId, LOD,
 						MeshVertices[LOD].Num(), TetHits, TriHits, Adoptions, NumOrphans);
 
 				} // end for all LOD
diff --git a/Engine/Plugins/Experimental/ChaosFlesh/Source/ChaosFleshNodes/Public/Dataflow/ChaosFleshBindingsNodes.h b/Engine/Plugins/Experimental/ChaosFlesh/Source/ChaosFleshNodes/Public/Dataflow/ChaosFleshBindingsNodes.h
--- a/Engine/Plugins/Experimental/ChaosFlesh/Source/ChaosFleshNodes/Public/Dataflow/ChaosFleshBindingsNodes.h
+++ b/Engine/Plugins/Experimental/ChaosFlesh/Source/ChaosFleshNodes/Public/Dataflow/ChaosFleshBindingsNodes.h
@@ -13,7 +13,7 @@ class USkeletalMesh;
 
 DECLARE_LOG_CATEGORY_EXTERN(LogMeshBindings, Verbose, All);
 
-// Generate barycentric bindings (used by the FleshDeformer deformer graph) of a render surface to a tetrahedral mesh.
+/** Generate barycentric bindings (used by the FleshDeformer deformer graph) of a render surface to a tetrahedral mesh. */
 USTRUCT(meta = (DataflowFlesh))
 struct FGenerateSurfaceBindings : public FDataflowNode
 {
@@ -24,30 +24,34 @@ struct FGenerateSurfaceBindings : public FDataflowNode
 public:
 	typedef FManagedArrayCollection DataType;
 
-	// Passthrough geometry collection. Bindings are stored as standalone groups in the \p Collection, keyed by the name of the input render mesh and all available LOD's.
+	/** Passthrough geometry collection. Bindings are stored as standalone groups in the \p Collection, keyed by the name of the input render mesh and all available LOD's. */
 	UPROPERTY(meta = (DataflowInput, DataflowOutput, DataflowPassthrough = "Collection", DisplayName = "Collection"))
 	FManagedArrayCollection Collection;
 
-	// The input mesh, whose render surface is used to generate bindings.
+	/** The input mesh, whose render surface is used to generate bindings. */
 	UPROPERTY(EditAnywhere, Category = "Dataflow", meta = (DataflowInput, DisplayName = "StaticMesh"))
 	TObjectPtr<const UStaticMesh> StaticMeshIn = nullptr;
 
-	// The input mesh, whose render surface is used to generate bindings.
+	/** The input mesh, whose render surface is used to generate bindings. */
 	UPROPERTY(EditAnywhere, Category = "Dataflow", meta = (DataflowInput, DisplayName = "SkeletalMesh"))
 	TObjectPtr<const USkeletalMesh> SkeletalMeshIn = nullptr;
 
+	/** Use the import geometry of the skeletal mesh. */
+	UPROPERTY(EditAnywhere, Category = "Dataflow", meta = (DisplayName = "UseSkeletalMeshImportModel"))
+	bool bUseSkeletalMeshImportModel = false;
+
 	UPROPERTY(meta = (DataflowInput, DisplayName = "(Optional)GeometryGroupGuidsIn"))
 	TArray<FString> GeometryGroupGuidsIn;
 
-	// Enable binding to the exterior hull of the tetrahedron mesh.
+	/** Enable binding to the exterior hull of the tetrahedron mesh. */
 	UPROPERTY(EditAnywhere, Category = "Dataflow", meta = (DisplayName = "DoSurfaceProjection"))
 	bool bDoSurfaceProjection = true;
 
-	// The maximum number of iterations to try expanding the domain while looking for surface triangles to bind to.
+	/** The maximum number of iterations to try expanding the domain while looking for surface triangles to bind to. */
 	UPROPERTY(EditAnywhere, Category = "Dataflow", meta = (DisplayName = "SurfaceProjectionIterations", EditCondition = "bDoSurfaceProjection == true"))
 	uint32 SurfaceProjectionIterations = 10;
 
-	// When nodes aren't contained in tetrahedra and surface projection fails, try to find suitable bindings by looking to neighboring parents.
+	/** When nodes aren't contained in tetrahedra and surface projection fails, try to find suitable bindings by looking to neighboring parents. */
 	UPROPERTY(EditAnywhere, Category = "Dataflow", meta = (DisplayName = "DoOrphanReparenting"))
 	bool bDoOrphanReparenting = true;
 
diff --git a/Engine/Source/Runtime/Experimental/Chaos/Private/GeometryCollection/Facades/CollectionTetrahedralBindingsFacade.cpp b/Engine/Source/Runtime/Experimental/Chaos/Private/GeometryCollection/Facades/CollectionTetrahedralBindingsFacade.cpp
--- a/Engine/Source/Runtime/Experimental/Chaos/Private/GeometryCollection/Facades/CollectionTetrahedralBindingsFacade.cpp
+++ b/Engine/Source/Runtime/Experimental/Chaos/Private/GeometryCollection/Facades/CollectionTetrahedralBindingsFacade.cpp
@@ -10,6 +10,100 @@
 
 namespace GeometryCollection::Facades
 {
+	// [ Duff et al. 2017, "Building an Orthonormal Basis, Revisited" ]
+	// Discontinuity at TangentZ.z == 0
+	Chaos::PMatrix<float, 3, 3>
+	FleshDeformerImpl::GetTangentBasis(const FVector3f& TangentZ)
+	{
+		const float Sign = copysignf(1.0f, TangentZ.Z);
+		const float A = -1.0f / (Sign + TangentZ.Z);
+		const float B = TangentZ.Z * TangentZ.Y * A;
+
+		FVector3f TangentX(
+			1.0f + Sign * TangentZ.X * TangentZ.X * A,
+			Sign * B,
+			-Sign * TangentZ.X);
+		FVector3f TangentY(
+			B,
+			Sign + TangentZ.Y * TangentZ.Y * A,
+			-TangentZ.Y);
+
+		return Chaos::PMatrix<float, 3, 3>(TangentX, TangentY, TangentZ);
+	}
+
+	Chaos::PMatrix<float, 3, 3>
+	FleshDeformerImpl::GetOrthogonalBasisVectors(
+		const FVector3f& PtA,
+		const FVector3f& PtB,
+		const FVector3f& PtC)
+	{
+		FVector3f EdgeBA = PtB - PtA;
+		FVector3f EdgeCA = PtC - PtA;
+		FVector3f OrthoNorm = FVector3f::CrossProduct(EdgeBA, EdgeCA).GetSafeNormal();
+		return GetTangentBasis(OrthoNorm);
+	}
+
+	FVector3f
+	FleshDeformerImpl::GetRotatedOffsetVector(
+		const FVector3f& Offset,
+		const FVector3f& RestPtA,
+		const FVector3f& RestPtB,
+		const FVector3f& RestPtC,
+		const FVector3f& CurrPtA,
+		const FVector3f& CurrPtB,
+		const FVector3f& CurrPtC)
+	{
+		Chaos::PMatrix<float, 3, 3> BasisVectors = GetOrthogonalBasisVectors(RestPtA, RestPtB, RestPtC);
+		Chaos::PMatrix<float, 3, 3> RestRotInv = BasisVectors.GetTransposed(); // Ortho matrix, so T == I
+		Chaos::PMatrix<float, 3, 3> CurrRot = GetOrthogonalBasisVectors(CurrPtA, CurrPtB, CurrPtC);
+		Chaos::PMatrix<float, 3, 3> BasisDelta = RestRotInv * CurrRot;
+		return BasisDelta.TransformVector(Offset); // ryan - mult order?
+	}
+
+	FVector3f
+	FleshDeformerImpl::GetRotatedOffsetVector(
+		const FIntVector4& Parents,
+		const FVector3f& Offset,
+		const TManagedArray<FVector3f>& RestVertices,
+		const TArray<Chaos::TVector<Chaos::FRealSingle, 3>>& CurrVertices)
+	{
+		return GetRotatedOffsetVector(
+			Offset,
+			RestVertices[Parents[0]], RestVertices[Parents[1]], RestVertices[Parents[2]],
+			CurrVertices[Parents[0]], CurrVertices[Parents[1]], CurrVertices[Parents[2]]);
+	}
+
+	FVector3f
+	FleshDeformerImpl::GetEmbeddedPosition(
+		const int32 SurfaceIndex,
+		const TManagedArrayAccessor<FIntVector4>* ParentsArray,
+		const TManagedArrayAccessor<FVector4f>* WeightsArray,
+		const TManagedArrayAccessor<FVector3f>* OffsetArray,
+		const TManagedArray<FVector3f>& RestVertices,
+		const TArray<Chaos::TVector<Chaos::FRealSingle, 3>>& CurrVertices)
+	{
+		const FIntVector4& Parents = ParentsArray->Get()[SurfaceIndex];
+		const FVector4f& Weights = WeightsArray->Get()[SurfaceIndex];
+
+		FVector3f Pos(0, 0, 0);
+
+		const int32 iEnd = Parents[3] == INDEX_NONE ? 3 : 4;
+		for (int32 i = 0; i < iEnd; i++)
+		{
+			Pos += CurrVertices[Parents[i]] * Weights[i];
+		}
+
+		// If surface binding, the last index is -1.
+		if (Parents[3] == INDEX_NONE)
+		{
+			const FVector3f& Offset = OffsetArray->Get()[SurfaceIndex];
+			Pos += GetRotatedOffsetVector(Parents, Offset, RestVertices, CurrVertices);
+		}
+
+		return Pos;
+	}
+
+
 	// Groups 
 	const FName FTetrahedralBindings::MeshBindingsGroupName = "MeshBindings";
 
@@ -21,6 +115,27 @@ namespace GeometryCollection::Facades
 	const FName FTetrahedralBindings::OffsetsAttributeName = "Offsets";
 	const FName FTetrahedralBindings::MaskAttributeName = "Mask";
 
+	// Static helpers
+	FString
+	FTetrahedralBindings::GetMeshId(const USkeletalMesh* SkeletalMesh, const bool bUseImportModel)
+	{
+		FPrimaryAssetId Id = SkeletalMesh->GetPrimaryAssetId();
+		FString MeshId = Id.IsValid() ? Id.ToString() : SkeletalMesh->GetName();
+		if (bUseImportModel)
+		{
+			MeshId.Append(TEXT("_ImportModel"));
+		}
+		return MeshId;
+	}
+
+	FString
+	FTetrahedralBindings::GetMeshId(const UStaticMesh* StaticMesh)
+	{
+		FPrimaryAssetId Id = StaticMesh->GetPrimaryAssetId();
+		FString MeshId = Id.IsValid() ? Id.ToString() : StaticMesh->GetName();
+		return MeshId;
+	}
+
 	FTetrahedralBindings::FTetrahedralBindings(FManagedArrayCollection& InCollection)
 		: MeshIdAttribute(InCollection, MeshIdAttributeName, MeshBindingsGroupName)
 	{}
@@ -45,7 +160,7 @@ namespace GeometryCollection::Facades
 			(Parents && Parents->IsValid()) && 
 			(Weights && Weights->IsValid()) && 
 			(Offsets && Offsets->IsValid()) &&
-			(Mask && Mask->IsValid());
+			(Masks && Masks->IsValid());
 	}
 
 	FName FTetrahedralBindings::GenerateMeshGroupName(
@@ -53,7 +168,8 @@ namespace GeometryCollection::Facades
 		const FName& MeshId,
 		const int32 LOD)
 	{
-		FString Str = FString::Printf(TEXT("TetrahedralBindings:TetMeshIdx:%d:%s:%d"), TetMeshIdx, *MeshId.ToString(), LOD);
+		const FString MeshIdStr = MeshId.GetPlainNameString();
+		FString Str = FString::Printf(TEXT("TetrahedralBindings:TetMeshIdx:%d:%s:%d"), TetMeshIdx, *MeshIdStr, LOD);
 		return FName(Str.Len(), *Str);
 	}
 
@@ -114,16 +230,16 @@ namespace GeometryCollection::Facades
 		Parents.Reset();
 		Weights.Reset();
 		Offsets.Reset();
-		Mask.Reset();
+		Masks.Reset();
 		FManagedArrayCollection& Collection = *MeshIdAttribute.GetCollection();
 		Parents.Reset(new TManagedArrayAccessor<FIntVector4>(Collection, ParentsAttributeName, GroupName));
 		Weights.Reset(new TManagedArrayAccessor<FVector4f>(Collection, WeightsAttributeName, GroupName));
 		Offsets.Reset(new TManagedArrayAccessor<FVector3f>(Collection, OffsetsAttributeName, GroupName));
-		Mask.Reset(new TManagedArrayAccessor<float>(Collection, MaskAttributeName, GroupName));
+		Masks.Reset(new TManagedArrayAccessor<float>(Collection, MaskAttributeName, GroupName));
 		Parents->Add(ManageArrayAccessor::EPersistencePolicy::MakePersistent, FGeometryCollection::VerticesGroup);
 		Weights->Add();
 		Offsets->Add();
-		Mask->Add();
+		Masks->Add();
 	}
 
 	bool FTetrahedralBindings::ReadBindingsGroup(const int32 TetMeshIdx, const FName& MeshId, const int32 LOD)
@@ -137,7 +253,7 @@ namespace GeometryCollection::Facades
 		Parents.Reset();
 		Weights.Reset();
 		Offsets.Reset();
-		Mask.Reset();
+		Masks.Reset();
 		if (!MeshIdAttribute.Find()->Contains(GroupName.ToString()))
 		{
 			return false;
@@ -149,7 +265,7 @@ namespace GeometryCollection::Facades
 			Parents.Reset(new TManagedArrayAccessor<FIntVector4>(*Collection, ParentsAttributeName, GroupName));
 			Weights.Reset(new TManagedArrayAccessor<FVector4f>(*Collection, WeightsAttributeName, GroupName));
 			Offsets.Reset(new TManagedArrayAccessor<FVector3f>(*Collection, OffsetsAttributeName, GroupName));
-			Mask.Reset(new TManagedArrayAccessor<float>(*Collection, MaskAttributeName, GroupName));
+			Masks.Reset(new TManagedArrayAccessor<float>(*Collection, MaskAttributeName, GroupName));
 		}
 		else
 		{
@@ -157,9 +273,9 @@ namespace GeometryCollection::Facades
 			Parents.Reset(new TManagedArrayAccessor<FIntVector4>(ConstCollection, ParentsAttributeName, GroupName));
 			Weights.Reset(new TManagedArrayAccessor<FVector4f>(ConstCollection, WeightsAttributeName, GroupName));
 			Offsets.Reset(new TManagedArrayAccessor<FVector3f>(ConstCollection, OffsetsAttributeName, GroupName));
-			Mask.Reset(new TManagedArrayAccessor<float>(ConstCollection, MaskAttributeName, GroupName));
+			Masks.Reset(new TManagedArrayAccessor<float>(ConstCollection, MaskAttributeName, GroupName));
 		}
-		return Parents->IsValid() && Weights->IsValid() && Offsets->IsValid() && Mask->IsValid();
+		return Parents->IsValid() && Weights->IsValid() && Offsets->IsValid() && Masks->IsValid();
 	}
 
 	void FTetrahedralBindings::RemoveBindingsGroup(const int32 TetMeshIdx, const FName& MeshId, const int32 LOD)
@@ -195,10 +311,10 @@ namespace GeometryCollection::Facades
 			Offsets->Remove();
 			Offsets.Reset();
 		}
-		if (Mask)
+		if (Masks)
 		{
-			Mask->Remove();
-			Mask.Reset();
+			Masks->Remove();
+			Masks.Reset();
 		}
 		// Only drop the group if it's empty at this point?
 		if (Collection.NumAttributes(GroupName) == 0)
@@ -224,7 +340,7 @@ namespace GeometryCollection::Facades
 		TManagedArray<FIntVector4>& ParentsValues = Parents->Modify();
 		TManagedArray<FVector4f>& WeightsValues = Weights->Modify();
 		TManagedArray<FVector3f>& OffsetsValues = Offsets->Modify();
-		TManagedArray<float>& MaskValues = Mask->Modify();
+		TManagedArray<float>& MaskValues = Masks->Modify();
 		for (int32 i = 0; i < Num; i++)
 		{
 			ParentsValues[i] = ParentsIn[i];
diff --git a/Engine/Source/Runtime/Experimental/Chaos/Public/GeometryCollection/Facades/CollectionTetrahedralBindingsFacade.h b/Engine/Source/Runtime/Experimental/Chaos/Public/GeometryCollection/Facades/CollectionTetrahedralBindingsFacade.h
--- a/Engine/Source/Runtime/Experimental/Chaos/Public/GeometryCollection/Facades/CollectionTetrahedralBindingsFacade.h
+++ b/Engine/Source/Runtime/Experimental/Chaos/Public/GeometryCollection/Facades/CollectionTetrahedralBindingsFacade.h
@@ -6,8 +6,51 @@
 #include "GeometryCollection/ManagedArrayCollection.h"
 #include "GeometryCollection/ManagedArrayAccessor.h"
 
+#include "Engine/StaticMesh.h"
+#include "Engine/SkeletalMesh.h"
+
 namespace GeometryCollection::Facades
 {
+	/**
+	* Flesh deformer behavior in c++ (rather than hlsl).
+	*/
+	class FleshDeformerImpl
+	{
+	public:
+		FleshDeformerImpl()
+		{}
+
+		/** Non-branching tangent basis vectors. Discontinuity at TangentZ.z == 0 */
+		static CHAOS_API Chaos::PMatrix<float, 3, 3> GetTangentBasis(const FVector3f& TangentZ);
+		/** Returns an orthogonal basis from triangle vertices. */
+		static CHAOS_API Chaos::PMatrix<float, 3, 3> GetOrthogonalBasisVectors(const FVector3f& PtA, const FVector3f& PtB, const FVector3f& PtC);
+
+		/**
+		* Given an \p Offset vector relative to rest triangle configuration \p RestPtA, \p RestPtB,
+		* \p RestPtC, returns a rotated offset vector relative to the current triangle configuration,
+		* \p CurrPtA, \p CurrPtB, CurrPtC.
+		*/
+		static CHAOS_API FVector3f GetRotatedOffsetVector(const FVector3f& Offset, const FVector3f& RestPtA, const FVector3f& RestPtB, const FVector3f& RestPtC, const FVector3f& CurrPtA, const FVector3f& CurrPtB, const FVector3f& CurrPtC);
+		static CHAOS_API FVector3f GetRotatedOffsetVector(const FIntVector4& Parents, const FVector3f& Offset, const TManagedArray<FVector3f>& RestVertices, const TArray<Chaos::TVector<Chaos::FRealSingle, 3>>& CurrVertices);
+
+		/**
+		* Evaluate the tetrahedral bindings for \p SurfaceIndex.
+		*
+		* \p SurfaceIndex - the bindings index to evaluate.
+		* \p ParentsArray - bindings parents buffer.
+		* \p WeightsArray - bindings weights buffer.
+		* \p OffsetArray - bindings offset buffer.
+		* \p RestVertices - tetrahedral mesh rest points.
+		* \p CurrVertices - tetrahedral mesh current points.
+		*/
+		static CHAOS_API FVector3f GetEmbeddedPosition(
+			const int32 SurfaceIndex,
+			const TManagedArrayAccessor<FIntVector4>* ParentsArray,
+			const TManagedArrayAccessor<FVector4f>* WeightsArray,
+			const TManagedArrayAccessor<FVector3f>* OffsetArray,
+			const TManagedArray<FVector3f>& RestVertices,
+			const TArray<Chaos::TVector<Chaos::FRealSingle, 3>>& CurrVertices);
+	};
 
 	/**
 	* FTetrahedralBindings
@@ -18,7 +61,6 @@ namespace GeometryCollection::Facades
 	class FTetrahedralBindings
  	{
 	public:
-
 		// groups
 		static CHAOS_API const FName MeshBindingsGroupName;
 
@@ -37,6 +79,13 @@ namespace GeometryCollection::Facades
 		//! Per vertex amount for deformer masking.
 		static CHAOS_API const FName MaskAttributeName;
 
+		//
+		// Static helpers
+		//
+
+		static CHAOS_API FString GetMeshId(const USkeletalMesh* SkeletalMesh, const bool bUseImportModel = false);
+		static CHAOS_API FString GetMeshId(const UStaticMesh* StaticMesh);
+
 		/**
 		* FSelectionFacade Constuctor
 		* @param VertixDependencyGroup : GroupName the index attribute is dependent on. 
@@ -108,7 +157,7 @@ namespace GeometryCollection::Facades
 		*    full deformation.
 		*/
 		CHAOS_API void SetBindingsData(const TArray<FIntVector4>& ParentsIn, const TArray<FVector4f>& WeightsIn, const TArray<FVector3f>& OffsetsIn, const TArray<float>& MaskIn);
-		void SetBindingsData(const TArray<FIntVector4>& ParentsIn, const TArray<FVector4f>& WeightsIn, const TArray<FVector3f>& OffsetsIn)
+		CHAOS_API void SetBindingsData(const TArray<FIntVector4>& ParentsIn, const TArray<FVector4f>& WeightsIn, const TArray<FVector3f>& OffsetsIn)
 		{
 			TArray<float> MaskTmp; MaskTmp.SetNum(ParentsIn.Num());
 			for (int32 i = 0; i < ParentsIn.Num(); i++)
@@ -136,8 +185,137 @@ namespace GeometryCollection::Facades
 		/**
 		* Get Mask array.
 		*/
-		const TManagedArrayAccessor<float>* GetMaskRO() const { return Mask.Get(); }
-		      TManagedArrayAccessor<float>* GetMask() { check(!IsConst()); return Mask.Get(); }
+		const TManagedArrayAccessor<float>* GetMaskRO() const { return Masks.Get(); }
+		      TManagedArrayAccessor<float>* GetMask() { check(!IsConst()); return Masks.Get(); }
+
+		/**
+		* Bindings evaluator.
+		* 
+		* Given binding parents, weights, offsets, and current tetrahedral mesh positions, get the
+		* current bound position.
+		*/
+		class Evaluator
+		{
+		public:
+			Evaluator(
+				const TManagedArrayAccessor<FIntVector4>* Parents, 
+				const TManagedArrayAccessor<FVector4f>* Weights, 
+				const TManagedArrayAccessor<FVector3f>* Offsets,
+				const TManagedArray<FVector3f>* RestVertices)
+				: MinIndexValue(TNumericLimits<int32>::Max())
+				, MaxIndexValue(INDEX_NONE)
+				, ParentsArray(Parents)
+				, WeightsArray(Weights)
+				, OffsetsArray(Offsets)
+				, RestVerticesArray(RestVertices)
+			{}
+
+			CHAOS_API bool IsValid() const
+			{ 
+				return ParentsArray && WeightsArray && OffsetsArray && 
+					(ParentsArray->Get().Num()==WeightsArray->Get().Num() && ParentsArray->Get().Num()==OffsetsArray->Get().Num()) &&
+					(RestVerticesArray->IsValidIndex(MinIndex()) && RestVerticesArray->IsValidIndex(MaxIndex())); 
+			}
+			CHAOS_API int32 NumVertices() const { return ParentsArray->Num(); }
+			CHAOS_API int32 MinIndex() const
+			{
+				if (MinIndexValue == TNumericLimits<int32>::Max())
+				{
+					for (const FIntVector4& Tet : ParentsArray->Get())
+					{
+						const int32 iEnd = Tet[3] == INDEX_NONE ? 3 : 4;
+						for (int32 i = 0; i < iEnd; i++) MinIndexValue = FMath::Min(MinIndexValue, Tet[i]);
+					}
+				}
+				return MinIndexValue;
+			}
+			CHAOS_API int32 MaxIndex() const
+			{
+				if (MaxIndexValue == INDEX_NONE)
+				{
+					for (const FIntVector4& Tet : ParentsArray->Get())
+					{
+						const int32 iEnd = Tet[3] == INDEX_NONE ? 3 : 4;
+						for (int32 i = 0; i < iEnd; i++) MaxIndexValue = FMath::Max(MaxIndexValue, Tet[i]);
+					}
+				}
+				return MaxIndexValue;
+			}
+			CHAOS_API FVector3f GetEmbeddedPosition(const int32 Index, const TArray<Chaos::TVector<Chaos::FRealSingle, 3>>& CurrVertices) const
+			{
+				return FleshDeformerImpl::GetEmbeddedPosition(Index, ParentsArray, WeightsArray, OffsetsArray, *RestVerticesArray, CurrVertices);
+			}
+
+		private:
+			mutable int32 MinIndexValue;
+			mutable int32 MaxIndexValue;
+			const TManagedArrayAccessor<FIntVector4>* ParentsArray;
+			const TManagedArrayAccessor<FVector4f>* WeightsArray;
+			const TManagedArrayAccessor<FVector3f>* OffsetsArray;
+			const TManagedArray<FVector3f>* RestVerticesArray;
+		};
+
+		/**
+		* Masked bindings evaluator.
+		* 
+		* Given binding parents, weights, offsets, mask weights, and current tetrahedral mesh positions, get 
+		* the current bound position blended with the rig position.
+		*/
+		class MaskedEvaluator
+		{
+		public:
+			MaskedEvaluator(
+				const TManagedArrayAccessor<FIntVector4>* Parents, 
+				const TManagedArrayAccessor<FVector4f>* Weights, 
+				const TManagedArrayAccessor<FVector3f>* Offsets, 
+				const TManagedArrayAccessor<float>* Masks,
+				const TManagedArray<FVector3f>* RestVertices)
+				: UnmaskedEval(Parents, Weights, Offsets, RestVertices)
+				, ParentsArray(Parents)
+				, MasksArray(Masks)
+			{}
+
+			CHAOS_API bool IsValid() const { return UnmaskedEval.IsValid() && (ParentsArray->Get().Num() == MasksArray->Get().Num()); }
+			CHAOS_API int32 NumVertices() const { return UnmaskedEval.NumVertices(); }
+			CHAOS_API int32 MinIndex() const { return UnmaskedEval.MinIndex(); }
+			CHAOS_API int32 MaxIndex() const { return UnmaskedEval.MaxIndex(); }
+			CHAOS_API FVector3f GetEmbeddedPosition(const int32 Index, const FVector3f& RigPosition, const TArray<Chaos::TVector<Chaos::FRealSingle, 3>>& CurrVertices) const
+			{
+				const float Mask = MasksArray->Get()[Index];
+				if (Mask < 1.0e-6)
+				{
+					return RigPosition;
+				}
+				else if (Mask > 1.0 - 1.0e-6)
+				{
+					return UnmaskedEval.GetEmbeddedPosition(Index, CurrVertices);
+				}
+				else
+				{
+					FVector3f EmbeddedPos = UnmaskedEval.GetEmbeddedPosition(Index, CurrVertices);
+					return (1.0 - Mask) * RigPosition + Mask * EmbeddedPos;
+				}
+			}
+		private:
+			Evaluator UnmaskedEval;
+			const TManagedArrayAccessor<FIntVector4>* ParentsArray;
+			const TManagedArrayAccessor<float>* MasksArray;
+		};
+
+		/**
+		* Initialize an \c Evaluator that samples bindings and computes resulting positions.
+		*/
+		CHAOS_API TUniquePtr<Evaluator> InitEvaluator(const TManagedArray<FVector3f>* RestVertices) const
+		{
+			return TUniquePtr<Evaluator>(new Evaluator(Parents.Get(), Weights.Get(), Offsets.Get(), RestVertices));
+		}
+
+		/**
+		* Initialize a \c MaskedEvaluator that samples bindings and computes resulting positions,
+		* masked and blended by a rig evaluated position.
+		*/
+		CHAOS_API TUniquePtr<MaskedEvaluator> InitMaskedEvaluator(const TManagedArray<FVector3f>* RestVertices) const
+		{ return TUniquePtr<MaskedEvaluator>(new MaskedEvaluator(Parents.Get(), Weights.Get(), Offsets.Get(), Masks.Get(), RestVertices)); }
 
 	private:
 		TManagedArrayAccessor<FString> MeshIdAttribute;
@@ -145,7 +323,7 @@ namespace GeometryCollection::Facades
 		TUniquePtr<TManagedArrayAccessor<FIntVector4>> Parents;
 		TUniquePtr<TManagedArrayAccessor<FVector4f>> Weights;
 		TUniquePtr<TManagedArrayAccessor<FVector3f>> Offsets;
-		TUniquePtr<TManagedArrayAccessor<float>> Mask;
+		TUniquePtr<TManagedArrayAccessor<float>> Masks;
 	};
 
 }
