Chaos : Fix engine dependency introduced in BindingsFacade

diff --git a/Engine/Plugins/Experimental/ChaosFlesh/Source/ChaosFleshEditor/ChaosFleshEditor.Build.cs b/Engine/Plugins/Experimental/ChaosFlesh/Source/ChaosFleshEditor/ChaosFleshEditor.Build.cs
--- a/Engine/Plugins/Experimental/ChaosFlesh/Source/ChaosFleshEditor/ChaosFleshEditor.Build.cs
+++ b/Engine/Plugins/Experimental/ChaosFlesh/Source/ChaosFleshEditor/ChaosFleshEditor.Build.cs
@@ -32,6 +32,7 @@ namespace UnrealBuildTool.Rules
 					"ChaosCachingUSD",
 					"ChaosFlesh",
 					"ChaosFleshEngine",
+					"ChaosFleshNodes",
 					"Core",
 					"CoreUObject",
 					"DataflowCore",
diff --git a/Engine/Plugins/Experimental/ChaosFlesh/Source/ChaosFleshEditor/Private/ChaosFlesh/Cmd/ChaosFleshCommands.cpp b/Engine/Plugins/Experimental/ChaosFlesh/Source/ChaosFleshEditor/Private/ChaosFlesh/Cmd/ChaosFleshCommands.cpp
--- a/Engine/Plugins/Experimental/ChaosFlesh/Source/ChaosFleshEditor/Private/ChaosFlesh/Cmd/ChaosFleshCommands.cpp
+++ b/Engine/Plugins/Experimental/ChaosFlesh/Source/ChaosFleshEditor/Private/ChaosFlesh/Cmd/ChaosFleshCommands.cpp
@@ -14,6 +14,8 @@
 #include "Chaos/Tetrahedron.h"
 #include "ChaosCache/FleshComponentCacheAdapter.h"
 
+#include "Dataflow/ChaosFleshBindingsNodes.h"
+
 #if USE_USD_SDK && DO_USD_CACHING
 #include "HAL/PlatformFile.h"
 #include "HAL/PlatformFileManager.h"
@@ -592,7 +594,7 @@ void FChaosFleshCommands::CreateGeometryCache(const TArray<FString>& Args, UWorl
 				// Extract bindings from tet mesh to skel mesh import geometry in the rest collection.
 				//
 
-				FString MeshId = GeometryCollection::Facades::FTetrahedralBindings::GetMeshId(SkeletalMesh, false);
+				FString MeshId = UE::TetrahedralBindingsEngineUtil::GetMeshId(SkeletalMesh, false);
 				FName MeshIdName(MeshId);
 				const int32 TetIndex = TetBindings.GetTetMeshIndex(MeshIdName, LODIndex);
 				if (TetIndex == INDEX_NONE)
diff --git a/Engine/Plugins/Experimental/ChaosFlesh/Source/ChaosFleshNodes/Private/Dataflow/ChaosFleshBindingsNodes.cpp b/Engine/Plugins/Experimental/ChaosFlesh/Source/ChaosFleshNodes/Private/Dataflow/ChaosFleshBindingsNodes.cpp
--- a/Engine/Plugins/Experimental/ChaosFlesh/Source/ChaosFleshNodes/Private/Dataflow/ChaosFleshBindingsNodes.cpp
+++ b/Engine/Plugins/Experimental/ChaosFlesh/Source/ChaosFleshNodes/Private/Dataflow/ChaosFleshBindingsNodes.cpp
@@ -32,6 +32,28 @@ namespace Dataflow
 	}
 }
 
+namespace UE::TetrahedralBindingsEngineUtil
+{
+	FString GetMeshId(const USkeletalMesh* SkeletalMesh, const bool bUseImportModel)
+	{
+		FPrimaryAssetId Id = SkeletalMesh->GetPrimaryAssetId();
+		FString MeshId = Id.IsValid() ? Id.ToString() : SkeletalMesh->GetName();
+		if (bUseImportModel)
+		{
+			MeshId.Append(TEXT("_ImportModel"));
+		}
+		return MeshId;
+	}
+
+	FString GetMeshId(const UStaticMesh* StaticMesh)
+	{
+		FPrimaryAssetId Id = StaticMesh->GetPrimaryAssetId();
+		FString MeshId = Id.IsValid() ? Id.ToString() : StaticMesh->GetName();
+		return MeshId;
+	}
+}
+
+
 void
 BuildVertexToVertexAdjacencyBuffer(
 	const UE::Geometry::FDynamicMesh3 DynamicMesh,
@@ -132,6 +154,8 @@ BuildVertexToVertexAdjacencyBuffer(
 	}
 }
 
+
+
 void
 FGenerateSurfaceBindings::Evaluate(Dataflow::FContext& Context, const FDataflowOutput* Out) const
 {
@@ -180,7 +204,7 @@ FGenerateSurfaceBindings::Evaluate(Dataflow::FContext& Context, const FDataflowO
 			if (SkeletalMesh)
 			{
 				FPrimaryAssetId Id = SkeletalMesh->GetPrimaryAssetId();
-				MeshId = GeometryCollection::Facades::FTetrahedralBindings::GetMeshId(SkeletalMesh, bUseSkeletalMeshImportModel);
+				MeshId = UE::TetrahedralBindingsEngineUtil::GetMeshId(SkeletalMesh, bUseSkeletalMeshImportModel);
 
 				if (!bUseSkeletalMeshImportModel)
 				{
@@ -258,7 +282,7 @@ FGenerateSurfaceBindings::Evaluate(Dataflow::FContext& Context, const FDataflowO
 			}
 			else // StaticMesh
 			{
-				MeshId = GeometryCollection::Facades::FTetrahedralBindings::GetMeshId(StaticMesh);
+				MeshId = UE::TetrahedralBindingsEngineUtil::GetMeshId(StaticMesh);
 
 				const FStaticMeshRenderData* RenderData = StaticMesh->GetRenderData();
 				const int32 NumLOD = RenderData->LODResources.Num();
diff --git a/Engine/Plugins/Experimental/ChaosFlesh/Source/ChaosFleshNodes/Public/Dataflow/ChaosFleshBindingsNodes.h b/Engine/Plugins/Experimental/ChaosFlesh/Source/ChaosFleshNodes/Public/Dataflow/ChaosFleshBindingsNodes.h
--- a/Engine/Plugins/Experimental/ChaosFlesh/Source/ChaosFleshNodes/Public/Dataflow/ChaosFleshBindingsNodes.h
+++ b/Engine/Plugins/Experimental/ChaosFlesh/Source/ChaosFleshNodes/Public/Dataflow/ChaosFleshBindingsNodes.h
@@ -6,6 +6,9 @@
 #include "Dataflow/DataflowEngine.h"
 #include "GeometryCollection/ManagedArrayCollection.h"
 
+#include "Engine/StaticMesh.h"
+#include "Engine/SkeletalMesh.h"
+
 #include "ChaosFleshBindingsNodes.generated.h"
 
 class UStaticMesh;
@@ -13,6 +16,12 @@ class USkeletalMesh;
 
 DECLARE_LOG_CATEGORY_EXTERN(LogMeshBindings, Verbose, All);
 
+namespace UE::TetrahedralBindingsEngineUtil
+{
+	CHAOSFLESHNODES_API FString GetMeshId(const USkeletalMesh* SkeletalMesh, const bool bUseImportModel = false);
+	CHAOSFLESHNODES_API FString GetMeshId(const UStaticMesh* StaticMesh);
+}
+
 /** Generate barycentric bindings (used by the FleshDeformer deformer graph) of a render surface to a tetrahedral mesh. */
 USTRUCT(meta = (DataflowFlesh))
 struct FGenerateSurfaceBindings : public FDataflowNode
diff --git a/Engine/Source/Runtime/Experimental/Chaos/Private/GeometryCollection/Facades/CollectionTetrahedralBindingsFacade.cpp b/Engine/Source/Runtime/Experimental/Chaos/Private/GeometryCollection/Facades/CollectionTetrahedralBindingsFacade.cpp
--- a/Engine/Source/Runtime/Experimental/Chaos/Private/GeometryCollection/Facades/CollectionTetrahedralBindingsFacade.cpp
+++ b/Engine/Source/Runtime/Experimental/Chaos/Private/GeometryCollection/Facades/CollectionTetrahedralBindingsFacade.cpp
@@ -115,27 +115,6 @@ namespace GeometryCollection::Facades
 	const FName FTetrahedralBindings::OffsetsAttributeName = "Offsets";
 	const FName FTetrahedralBindings::MaskAttributeName = "Mask";
 
-	// Static helpers
-	FString
-	FTetrahedralBindings::GetMeshId(const USkeletalMesh* SkeletalMesh, const bool bUseImportModel)
-	{
-		FPrimaryAssetId Id = SkeletalMesh->GetPrimaryAssetId();
-		FString MeshId = Id.IsValid() ? Id.ToString() : SkeletalMesh->GetName();
-		if (bUseImportModel)
-		{
-			MeshId.Append(TEXT("_ImportModel"));
-		}
-		return MeshId;
-	}
-
-	FString
-	FTetrahedralBindings::GetMeshId(const UStaticMesh* StaticMesh)
-	{
-		FPrimaryAssetId Id = StaticMesh->GetPrimaryAssetId();
-		FString MeshId = Id.IsValid() ? Id.ToString() : StaticMesh->GetName();
-		return MeshId;
-	}
-
 	FTetrahedralBindings::FTetrahedralBindings(FManagedArrayCollection& InCollection)
 		: MeshIdAttribute(InCollection, MeshIdAttributeName, MeshBindingsGroupName)
 	{}
diff --git a/Engine/Source/Runtime/Experimental/Chaos/Public/GeometryCollection/Facades/CollectionTetrahedralBindingsFacade.h b/Engine/Source/Runtime/Experimental/Chaos/Public/GeometryCollection/Facades/CollectionTetrahedralBindingsFacade.h
--- a/Engine/Source/Runtime/Experimental/Chaos/Public/GeometryCollection/Facades/CollectionTetrahedralBindingsFacade.h
+++ b/Engine/Source/Runtime/Experimental/Chaos/Public/GeometryCollection/Facades/CollectionTetrahedralBindingsFacade.h
@@ -6,9 +6,6 @@
 #include "GeometryCollection/ManagedArrayCollection.h"
 #include "GeometryCollection/ManagedArrayAccessor.h"
 
-#include "Engine/StaticMesh.h"
-#include "Engine/SkeletalMesh.h"
-
 namespace GeometryCollection::Facades
 {
 	/**
@@ -79,13 +76,6 @@ namespace GeometryCollection::Facades
 		//! Per vertex amount for deformer masking.
 		static CHAOS_API const FName MaskAttributeName;
 
-		//
-		// Static helpers
-		//
-
-		static CHAOS_API FString GetMeshId(const USkeletalMesh* SkeletalMesh, const bool bUseImportModel = false);
-		static CHAOS_API FString GetMeshId(const UStaticMesh* StaticMesh);
-
 		/**
 		* FSelectionFacade Constuctor
 		* @param VertixDependencyGroup : GroupName the index attribute is dependent on. 
