ChaosFlesh: Fixed kinematic bindings from the rig

#virtualized

diff --git a/Engine/Plugins/Experimental/ChaosFlesh/Source/ChaosFleshNodes/Private/Dataflow/ChaosFleshKinematicInitializationNodes.cpp b/Engine/Plugins/Experimental/ChaosFlesh/Source/ChaosFleshNodes/Private/Dataflow/ChaosFleshKinematicInitializationNodes.cpp
--- a/Engine/Plugins/Experimental/ChaosFlesh/Source/ChaosFleshNodes/Private/Dataflow/ChaosFleshKinematicInitializationNodes.cpp
+++ b/Engine/Plugins/Experimental/ChaosFlesh/Source/ChaosFleshNodes/Private/Dataflow/ChaosFleshKinematicInitializationNodes.cpp
@@ -50,14 +50,15 @@ void FKinematicTetrahedralBindingsDataflowNode::Evaluate(Dataflow::FContext& Con
 			//parse exclusion list to find bones to skip
 			TArray<FString> StrArray;
 			ExclusionList.ParseIntoArray(StrArray, *FString(" "));				
-			TSet<int32> AddedVertSet;
 			int32 NumTets = Tetrahedron->Num();
 			TArray<FTransform> ComponentPose;
 			Dataflow::Animation::GlobalTransforms(SkeletalMesh->GetRefSkeleton(), ComponentPose);
-			for (int32 b = SkeletalMesh->GetRefSkeleton().GetNum()-1; b > -1; b--)
+			TArray<bool> VertAdded;
+			VertAdded.Init(false, Vertex->Num());
+			auto DoSkipBoneIndex = [&StrArray, &SkeletalMesh](int32 BoneIndex)
 			{
 				bool Skip = false;
-				FString BoneName = SkeletalMesh->GetRefSkeleton().GetBoneName(b).ToString();
+				FString BoneName = SkeletalMesh->GetRefSkeleton().GetBoneName(BoneIndex).ToString();
 				for (FString Elem : StrArray)
 				{
 					if (BoneName.Contains(Elem))
@@ -66,13 +67,13 @@ void FKinematicTetrahedralBindingsDataflowNode::Evaluate(Dataflow::FContext& Con
 						break;
 					}
 				}
-				if (Skip)
-					continue;
-
+				return Skip;
+			};
+			for (int32 b = 0; b < SkeletalMesh->GetRefSkeleton().GetNum(); ++b)
+			{
 				FVector3f BonePosition(ComponentPose[b].GetTranslation());
 				int32 ParentIndex=SkeletalMesh->GetRefSkeleton().GetParentIndex(b);
-				
-				if (ParentIndex != INDEX_NONE) 
+				if (!(ParentIndex == INDEX_NONE || DoSkipBoneIndex(b) || DoSkipBoneIndex(ParentIndex)))
 				{
 					FVector3f ParentPosition(ComponentPose[ParentIndex].GetTranslation());
 					FVector3f RayDir = ParentPosition - BonePosition;
@@ -103,10 +104,10 @@ void FKinematicTetrahedralBindingsDataflowNode::Evaluate(Dataflow::FContext& Con
 							if (KeepTet) 
 							{	
 								for (int32 c = 0; c < 4; ++c)
-								{
-									if (!AddedVertSet.Contains((*Tetrahedron)[t][c]))
+								{	
+									if (!VertAdded[(*Tetrahedron)[t][c]])
 									{
-										AddedVertSet.Add((*Tetrahedron)[t][c]);
+										VertAdded[(*Tetrahedron)[t][c]] = true;
 										BoneVertSet.Add((*Tetrahedron)[t][c]);
 									}
 								}
@@ -123,17 +124,16 @@ void FKinematicTetrahedralBindingsDataflowNode::Evaluate(Dataflow::FContext& Con
 							FKinematics Kinematics(InCollection); Kinematics.DefineSchema();
 							if (Kinematics.IsValid())
 							{
-								FKinematics::FBindingKey Binding = Kinematics.SetBoneBindings(b, BoundVerts, BoundWeights);
+								FKinematics::FBindingKey Binding = Kinematics.SetBoneBindings(ParentIndex, BoundVerts, BoundWeights);
 								TManagedArray<TArray<FVector3f>>& LocalPos = InCollection.AddAttribute<TArray<FVector3f>>("LocalPosition", Binding.GroupName);
 								Kinematics.AddKinematicBinding(Binding);
-
-								auto FloatVert = [](FVector3d V) { return FVector3f(V.X, V.Y, V.Z); };
 								auto DoubleVert = [](FVector3f V) { return FVector3d(V.X, V.Y, V.Z); };
+								auto FloatVert = [](FVector3d V) { return FVector3f(V.X, V.Y, V.Z); };
 								LocalPos[Binding.Index].SetNum(BoundVerts.Num());
 								for (int32 i = 0; i < BoundVerts.Num(); i++)
 								{
 									FVector3f Temp = (*Vertex)[BoundVerts[i]];
-									LocalPos[Binding.Index][i] = FloatVert(ComponentPose[b].InverseTransformPosition(DoubleVert(Temp)));
+									LocalPos[Binding.Index][i] = FloatVert(ComponentPose[ParentIndex].InverseTransformPosition(DoubleVert(Temp)));
 								}
 							}
 						}
