ChaosFlesh : Set FrameRate of GeometryCache when converting from USD Flesh cache

diff --git a/Engine/Plugins/Experimental/ChaosFlesh/Source/ChaosFleshEditor/Private/ChaosFlesh/Cmd/ChaosFleshCommands.cpp b/Engine/Plugins/Experimental/ChaosFlesh/Source/ChaosFleshEditor/Private/ChaosFlesh/Cmd/ChaosFleshCommands.cpp
--- a/Engine/Plugins/Experimental/ChaosFlesh/Source/ChaosFleshEditor/Private/ChaosFlesh/Cmd/ChaosFleshCommands.cpp
+++ b/Engine/Plugins/Experimental/ChaosFlesh/Source/ChaosFleshEditor/Private/ChaosFlesh/Cmd/ChaosFleshCommands.cpp
@@ -418,6 +418,12 @@ void FChaosFleshCommands::CreateGeometryCache(const TArray<FString>& Args, UWorl
 	FString UsdFileOverride;
 	bool bUsdFile = ArgParse::ArgStringValue(Args, FString(TEXT("UsdFile")), UsdFileOverride);
 
+	float FrameRate = 24.;
+	ArgParse::ArgFloatValue(Args, FString(TEXT("FrameRate")), FrameRate);
+
+	int MaxNumFrames = INT_MAX;
+	ArgParse::ArgIntValue(Args, FString(TEXT("MaxNumFrames")), MaxNumFrames);
+
 	// Find a ChaosCacheManager, if a cache file hasn't been specified.
 	AChaosCacheManager* CacheManager = nullptr;
 	if (bUsdFile)
@@ -651,6 +657,26 @@ void FChaosFleshCommands::CreateGeometryCache(const TArray<FString>& Args, UWorl
 					continue;
 				}
 
+				if (TimeSamples.Num() > 1)
+				{
+					float MinTime = TimeSamples[0];
+					float MaxTime = TimeSamples[TimeSamples.Num() - 1];
+					float DeltaTime = 1. / FMath::Max(1.,FrameRate);
+					float TotalCacheTime = MaxTime - MinTime;
+
+					int NumSamples = FMath::Min(MaxNumFrames, TotalCacheTime / DeltaTime);
+					TimeSamples.SetNumUninitialized(NumSamples);
+
+					float CurrentTime = MinTime;
+					for (int SampleIdx = 0; SampleIdx < NumSamples; SampleIdx++)
+					{
+						ensure(MinTime <= CurrentTime && CurrentTime <= MaxTime);
+						TimeSamples[SampleIdx] = CurrentTime;
+						CurrentTime += DeltaTime;
+					}
+				}
+
+
 				//
 				// Deform render geometry, storing per frame data in CurrFramePositions.
 				//
@@ -695,7 +721,9 @@ void FChaosFleshCommands::CreateGeometryCache(const TArray<FString>& Args, UWorl
 					TEXT("CreateGeometryCache - FleshComponent '%s' writing deformed SkeleltalMesh '%s' render geometry to geometry cache: '%s'"),
 					*FleshComponent->GetName(), *SkeletalMesh->GetName(), *GeometryCache->GetName());
 
-				UE::GeometryCacheHelpers::FGeometryCacheConstantTopologyWriter Writer(*GeometryCache);
+				UE::GeometryCacheHelpers::FGeometryCacheConstantTopologyWriter::FConfig Config;
+				Config.FPS = FrameRate;
+				UE::GeometryCacheHelpers::FGeometryCacheConstantTopologyWriter Writer(*GeometryCache, Config);
 
 				// Writes Indices, UV's, Colors, ImportedVertexNumbers, and "BatchesInfo".
 				const int32 Index = UE::GeometryCacheHelpers::AddTrackWriterFromSkinnedAsset(Writer, *SkeletalMesh);
diff --git a/Engine/Plugins/Experimental/ChaosFlesh/Source/ChaosFleshEditor/Public/ChaosFlesh/Cmd/ChaosFleshCommands.h b/Engine/Plugins/Experimental/ChaosFlesh/Source/ChaosFleshEditor/Public/ChaosFlesh/Cmd/ChaosFleshCommands.h
--- a/Engine/Plugins/Experimental/ChaosFlesh/Source/ChaosFleshEditor/Public/ChaosFlesh/Cmd/ChaosFleshCommands.h
+++ b/Engine/Plugins/Experimental/ChaosFlesh/Source/ChaosFleshEditor/Public/ChaosFlesh/Cmd/ChaosFleshCommands.h
@@ -39,6 +39,8 @@ public:
 	* needs to contain deformer bindings for the import geometry.
 	* @param Args - Supported arguments:
 	*   'UsdFile </path/to/file.usd>' overrides which usd file to use, rather than inferring it from a chaos cache manager.
+	*   'FrameRate 24' overrides the default 24 FPS frame rate for the output cache samples.
+	*   'MaxNumFrames <int>' clamps the number of frames to MaxNumFrames.
 	*/
 	static void CreateGeometryCache(const TArray<FString>& Args, UWorld* World);
 
