ChaosFlesh : Create tetrahedron against named geometry in a collection.

diff --git a/Engine/Plugins/Experimental/ChaosFlesh/Source/ChaosFleshNodes/Private/Dataflow/ChaosFleshCreateTetrahedronNode.cpp b/Engine/Plugins/Experimental/ChaosFlesh/Source/ChaosFleshNodes/Private/Dataflow/ChaosFleshCreateTetrahedronNode.cpp
new file mode 100644
--- /dev/null
+++ b/Engine/Plugins/Experimental/ChaosFlesh/Source/ChaosFleshNodes/Private/Dataflow/ChaosFleshCreateTetrahedronNode.cpp
@@ -0,0 +1,243 @@
+// Copyright Epic Games, Inc. All Rights Reserved.
+
+#include "Dataflow/ChaosFleshCreateTetrahedronNode.h"
+#include "Dataflow/ChaosFleshEngineAssetNodes.h"
+
+#include "Async/ParallelFor.h"
+#include "Chaos/Deformable/Utilities.h"
+#include "ChaosFlesh/ChaosFlesh.h"
+#include "Chaos/Tetrahedron.h"
+#include "Chaos/Utilities.h"
+#include "Chaos/UniformGrid.h"
+#include "ChaosFlesh/FleshCollection.h"
+#include "ChaosFlesh/FleshCollectionUtility.h"
+#include "ChaosLog.h"
+#include "Dataflow/DataflowInputOutput.h"
+#include "Dataflow/ChaosFleshNodesUtility.h"
+#include "DynamicMesh/DynamicMesh3.h"
+#include "DynamicMesh/DynamicMeshAABBTree3.h"
+#include "Engine/SkeletalMesh.h"
+#include "Engine/StaticMesh.h"
+#include "FTetWildWrapper.h"
+#include "Generate/IsosurfaceStuffing.h"
+#include "GeometryCollection/ManagedArrayCollection.h"
+#include "GeometryCollection/Facades/CollectionTetrahedralMetricsFacade.h"
+#include "GeometryCollection/GeometryCollectionAlgo.h"
+#include "MeshDescription.h"
+#include "MeshDescriptionToDynamicMesh.h"
+#include "Rendering/SkeletalMeshLODImporterData.h"
+#include "Rendering/SkeletalMeshModel.h"
+#include "SkeletalMeshLODModelToDynamicMesh.h" // MeshModelingBlueprints
+#include "Spatial/FastWinding.h"
+#include "Spatial/MeshAABBTree3.h"
+#include "Dataflow/ChaosFleshNodesUtility.h"
+
+//=============================================================================
+// FCreateTetrahedronDataflowNode
+//=============================================================================
+
+
+void FCreateTetrahedronDataflowNode::Evaluate(Dataflow::FContext& Context, const FDataflowOutput* Out) const
+{
+	if (Out->IsA<DataType>(&Collection))
+	{
+		TUniquePtr<FFleshCollection> InCollection(GetValue<DataType>(Context, &Collection).NewCopy<FFleshCollection>());
+
+		// TransformGroup
+		int32 NumTransforms = InCollection->NumElements(FTransformCollection::TransformGroup);
+		TManagedArray<FString>* TransformName = InCollection->FindAttribute<FString>("BoneName", FTransformCollection::TransformGroup);
+		TManagedArray<int32>* TransformToGeometryIndex = InCollection->FindAttribute<int32>("TransformToGeometryIndex", FTransformCollection::TransformGroup);
+		// Geometry Group
+		int32 NumGeometry = InCollection->NumElements(FGeometryCollection::GeometryGroup);
+		TManagedArray<int32>* GroupTransformIndex = InCollection->FindAttribute<int32>("TransformIndex", FGeometryCollection::GeometryGroup);
+		TManagedArray<int32>* VertexCount = InCollection->FindAttribute<int32>("VertexCount", FGeometryCollection::GeometryGroup);
+		TManagedArray<int32>* VertexStart = InCollection->FindAttribute<int32>("VertexStart", FGeometryCollection::GeometryGroup);
+		TManagedArray<int32>* FaceCount = InCollection->FindAttribute<int32>("FaceCount", FGeometryCollection::GeometryGroup);
+		TManagedArray<int32>* FaceStart = InCollection->FindAttribute<int32>("FaceStart", FGeometryCollection::GeometryGroup);
+		// Vertices Group
+		int32 NumVertices = InCollection->NumElements(FGeometryCollection::VerticesGroup);
+		const TManagedArray<FVector3f>* Vertex = InCollection->FindAttribute<FVector3f>("Vertex", FGeometryCollection::VerticesGroup);
+		// Faces Group
+		int32 NumTriangles = InCollection->NumElements(FGeometryCollection::FacesGroup);
+		const TManagedArray<FIntVector>* Faces = InCollection->FindAttribute<FIntVector>("Indices", FGeometryCollection::FacesGroup);
+
+		if (TransformName && TransformToGeometryIndex && GroupTransformIndex && VertexCount && VertexStart && FaceCount && FaceStart
+			&& Vertex && Faces)
+		{
+			TArray<int32> ProcessGeometryIndices = Dataflow::GetMatchingMeshIndices(MeshNames, InCollection.Get());
+
+			TArray<TUniquePtr<FFleshCollection>> CollectionBuffer;
+			for (int32 Gdx = 0; Gdx < NumGeometry; Gdx++)
+			{
+				if (ProcessGeometryIndices.Contains(Gdx)) 
+					CollectionBuffer.Add(TUniquePtr<FFleshCollection>( new FFleshCollection()));
+				else 
+					CollectionBuffer.Add(nullptr);
+			}
+
+			ParallelFor(ProcessGeometryIndices.Num(), [&](int32 i)
+			{
+				int32 Gdx = ProcessGeometryIndices[i];
+				FFleshCollection& TetCollection = *CollectionBuffer[Gdx];
+
+				UE::Geometry::FDynamicMesh3 DynamicMesh;
+				int32 vStart = (*VertexStart)[Gdx], vEnd = vStart + (*VertexCount)[Gdx];
+				for (int Vdx = vStart; Vdx < vEnd; Vdx++) DynamicMesh.AppendVertex(FVector((*Vertex)[Vdx]));
+				int32 fStart = (*FaceStart)[Gdx], fEnd = fStart + (*FaceCount)[Gdx];
+				for (int Fdx = fStart; Fdx < fEnd; Fdx++) DynamicMesh.AppendTriangle((*Faces)[Fdx] - FIntVector(vStart));
+				DynamicMesh.CompactInPlace();
+
+				if (Method == TetMeshingMethod::IsoStuffing)
+				{
+					EvaluateIsoStuffing(Context, TetCollection, DynamicMesh);
+				}
+				else if (Method == TetMeshingMethod::TetWild)
+				{
+					EvaluateTetWild(Context, TetCollection, DynamicMesh);
+				}
+
+				if (TetCollection.NumElements(FGeometryCollection::VerticesGroup))
+				{
+					TSet<int32> VertexToDeleteSet;
+					GeometryCollectionAlgo::ComputeStaleVertices(&TetCollection, VertexToDeleteSet);
+					TArray<int32> SortedVertices = VertexToDeleteSet.Array(); SortedVertices.Sort();
+					if (VertexToDeleteSet.Num()) TetCollection.RemoveElements(FGeometryCollection::VerticesGroup, SortedVertices);
+				}
+			});
+			for (int32 Gdx = 0; Gdx < NumGeometry; Gdx++)
+			{
+				if (CollectionBuffer[Gdx])
+				{
+					int32 GeomIndex = InCollection->AppendGeometry(*CollectionBuffer[Gdx]);
+					FString ParentName = FString::Printf(TEXT("%d"), GeomIndex);
+					if (0 <= (*GroupTransformIndex)[Gdx]
+						&& (*GroupTransformIndex)[Gdx] < InCollection->NumElements(FGeometryCollection::GeometryGroup)
+						&& !(*TransformName)[(*GroupTransformIndex)[Gdx]].IsEmpty())
+					{
+						ParentName = (*TransformName)[(*GroupTransformIndex)[Gdx]];
+					}
+					if (0 <= (*GroupTransformIndex)[GeomIndex]
+						&& (*GroupTransformIndex)[GeomIndex] < InCollection->NumElements(FGeometryCollection::GeometryGroup))
+					{
+						(*TransformName)[(*GroupTransformIndex)[GeomIndex]] = FString::Printf(TEXT("%s_TET"), *ParentName);
+					}
+				}
+			}
+		}
+		SetValue<const DataType&>(Context, *InCollection, &Collection);
+	}
+}
+
+void FCreateTetrahedronDataflowNode::EvaluateIsoStuffing(
+	Dataflow::FContext& Context, 
+	FFleshCollection& InCollection,
+	const UE::Geometry::FDynamicMesh3& DynamicMesh) const
+{
+#if WITH_EDITORONLY_DATA
+	if (NumCells > 0 && (-.5 <= OffsetPercent && OffsetPercent <= 0.5))
+	{
+		// Tet mesh generation
+		UE::Geometry::TIsosurfaceStuffing<double> IsosurfaceStuffing;
+		UE::Geometry::FDynamicMeshAABBTree3 Spatial(&DynamicMesh);
+		UE::Geometry::TFastWindingTree<UE::Geometry::FDynamicMesh3> FastWinding(&Spatial);
+		UE::Geometry::FAxisAlignedBox3d Bounds = Spatial.GetBoundingBox();
+		IsosurfaceStuffing.Bounds = FBox(Bounds);
+		double CellSize = Bounds.MaxDim() / NumCells;
+		IsosurfaceStuffing.CellSize = CellSize;
+		IsosurfaceStuffing.IsoValue = .5 + OffsetPercent;
+		IsosurfaceStuffing.Implicit = [&FastWinding, &Spatial](FVector3d Pos)
+		{
+			FVector3d Nearest = Spatial.FindNearestPoint(Pos);
+			double WindingSign = FastWinding.FastWindingNumber(Pos) - .5;
+			return FVector3d::Distance(Nearest, Pos) * FMathd::SignNonZero(WindingSign);
+		};
+
+		UE_LOG(LogChaosFlesh, Display, TEXT("Generating tet mesh via IsoStuffing..."));
+		IsosurfaceStuffing.Generate();
+		if (IsosurfaceStuffing.Tets.Num() > 0)
+		{
+			TArray<FVector> Vertices; Vertices.SetNumUninitialized(IsosurfaceStuffing.Vertices.Num());
+			TArray<FIntVector4> Elements; Elements.SetNumUninitialized(IsosurfaceStuffing.Tets.Num());
+			TArray<FIntVector3> SurfaceElements = Dataflow::GetSurfaceTriangles(IsosurfaceStuffing.Tets, !bDiscardInteriorTriangles);
+
+			for (int32 Tdx = 0; Tdx < IsosurfaceStuffing.Tets.Num(); ++Tdx)
+			{
+				Elements[Tdx] = IsosurfaceStuffing.Tets[Tdx];
+			}
+			for (int32 Vdx = 0; Vdx < IsosurfaceStuffing.Vertices.Num(); ++Vdx)
+			{
+				Vertices[Vdx] = IsosurfaceStuffing.Vertices[Vdx];
+			}
+
+			TUniquePtr<FTetrahedralCollection> TetCollection(FTetrahedralCollection::NewTetrahedralCollection(Vertices, SurfaceElements, Elements));
+			InCollection.AppendGeometry(*TetCollection.Get());
+
+			UE_LOG(LogChaosFlesh, Display,
+				TEXT("Generated tet mesh via IsoStuffing, num vertices: %d num tets: %d"), Vertices.Num(), Elements.Num());
+		}
+		else
+		{
+			UE_LOG(LogChaosFlesh, Warning, TEXT("IsoStuffing produced 0 tetrahedra."));
+		}
+	}
+#else
+	ensureMsgf(false, TEXT("FCreateTetrahedronDataflowNodes is an editor only node."));
+#endif
+}
+
+void FCreateTetrahedronDataflowNode::EvaluateTetWild(
+	Dataflow::FContext& Context, 
+	FFleshCollection& InCollection,
+	const UE::Geometry::FDynamicMesh3& DynamicMesh) const
+{
+#if WITH_EDITORONLY_DATA
+	if (/* placeholder for conditions for exec */true)
+	{
+		// Pull out Vertices and Triangles
+		TArray<FVector> Verts;
+		TArray<FIntVector3> Tris;
+		for (FVector V : DynamicMesh.VerticesItr())
+		{
+			Verts.Add(V);
+		}
+		for (UE::Geometry::FIndex3i Tri : DynamicMesh.TrianglesItr())
+		{
+			Tris.Emplace(Tri.A, Tri.B, Tri.C);
+		}
+
+		// Tet mesh generation
+		UE::Geometry::FTetWild::FTetMeshParameters Params;
+		Params.bCoarsen = bCoarsen;
+		Params.bExtractManifoldBoundarySurface = bExtractManifoldBoundarySurface;
+		Params.bSkipSimplification = bSkipSimplification;
+
+		Params.EpsRel = EpsRel;
+		Params.MaxIts = MaxIterations;
+		Params.StopEnergy = StopEnergy;
+		Params.IdealEdgeLength = IdealEdgeLength;
+
+		Params.bInvertOutputTets = bInvertOutputTets;
+
+		TArray<FVector> TetVerts;
+		TArray<FIntVector4> Tets;
+		FProgressCancel Progress;
+		UE_LOG(LogChaosFlesh, Display,TEXT("Generating tet mesh via TetWild..."));
+		if (UE::Geometry::FTetWild::ComputeTetMesh(Params, Verts, Tris, TetVerts, Tets, &Progress))
+		{
+			TArray<FIntVector3> SurfaceElements = Dataflow::GetSurfaceTriangles(Tets, !bDiscardInteriorTriangles);
+			TUniquePtr<FTetrahedralCollection> TetCollection(FTetrahedralCollection::NewTetrahedralCollection(TetVerts, SurfaceElements, Tets));
+			InCollection.AppendGeometry(*TetCollection.Get());
+
+			UE_LOG(LogChaosFlesh, Display,
+				TEXT("Generated tet mesh via TetWild, num vertices: %d num tets: %d"), TetVerts.Num(), Tets.Num());
+		}
+		else
+		{
+			UE_LOG(LogChaosFlesh, Error,
+				TEXT("TetWild tetrahedral mesh generation failed."));
+		}
+	}
+#else
+	ensureMsgf(false, TEXT("FCreateTetrahedronDataflowNodes is an editor only node."));
+#endif
+}
diff --git a/Engine/Plugins/Experimental/ChaosFlesh/Source/ChaosFleshNodes/Private/Dataflow/ChaosFleshImportGEO.cpp b/Engine/Plugins/Experimental/ChaosFlesh/Source/ChaosFleshNodes/Private/Dataflow/ChaosFleshImportGEO.cpp
--- a/Engine/Plugins/Experimental/ChaosFlesh/Source/ChaosFleshNodes/Private/Dataflow/ChaosFleshImportGEO.cpp
+++ b/Engine/Plugins/Experimental/ChaosFlesh/Source/ChaosFleshNodes/Private/Dataflow/ChaosFleshImportGEO.cpp
@@ -1,7 +1,7 @@
 // Copyright Epic Games, Inc. All Rights Reserved.
 
 #include "Dataflow/ChaosFleshImportGEO.h"
-#include "Dataflow/ChaosFleshTetrahedralNodes.h" // for GetSurfaceTriangles()
+#include "Dataflow/ChaosFleshNodesUtility.h"
 #include "ChaosFlesh/TetrahedralCollection.h"
 #include "ChaosFlesh/FleshCollection.h"
 
diff --git a/Engine/Plugins/Experimental/ChaosFlesh/Source/ChaosFleshNodes/Private/Dataflow/ChaosFleshNodesPlugin.cpp b/Engine/Plugins/Experimental/ChaosFlesh/Source/ChaosFleshNodes/Private/Dataflow/ChaosFleshNodesPlugin.cpp
--- a/Engine/Plugins/Experimental/ChaosFlesh/Source/ChaosFleshNodes/Private/Dataflow/ChaosFleshNodesPlugin.cpp
+++ b/Engine/Plugins/Experimental/ChaosFlesh/Source/ChaosFleshNodes/Private/Dataflow/ChaosFleshNodesPlugin.cpp
@@ -5,6 +5,7 @@
 
 #include "Dataflow/ChaosFleshBindingsNodes.h"
 #include "Dataflow/ChaosFleshCoreNodes.h"
+#include "Dataflow/ChaosFleshCreateTetrahedronNode.h"
 #include "Dataflow/ChaosFleshEngineAssetNodes.h"
 #include "Dataflow/ChaosFleshFiberDirectionInitializationNodes.h"
 #include "Dataflow/ChaosFleshImportGEO.h"
@@ -32,6 +33,7 @@ void IChaosFleshNodesPlugin::StartupModule()
 	Dataflow::ChaosFleshSkeletalBindingsNode();
 	Dataflow::ChaosFleshRadialTetrahedronNodes();
 	Dataflow::RegisterChaosFleshImportGEONodes();
+	DATAFLOW_NODE_REGISTER_CREATION_FACTORY(FCreateTetrahedronDataflowNode);
 }
 
 void IChaosFleshNodesPlugin::ShutdownModule()
diff --git a/Engine/Plugins/Experimental/ChaosFlesh/Source/ChaosFleshNodes/Private/Dataflow/ChaosFleshNodesUtility.cpp b/Engine/Plugins/Experimental/ChaosFlesh/Source/ChaosFleshNodes/Private/Dataflow/ChaosFleshNodesUtility.cpp
new file mode 100644
--- /dev/null
+++ b/Engine/Plugins/Experimental/ChaosFlesh/Source/ChaosFleshNodes/Private/Dataflow/ChaosFleshNodesUtility.cpp
@@ -0,0 +1,86 @@
+// Copyright Epic Games, Inc. All Rights Reserved.
+
+#include "Dataflow/ChaosFleshNodesUtility.h"
+
+#include "Chaos/Deformable/Utilities.h"
+#include "Dataflow/DataflowEngine.h"
+#include "GeometryCollection/GeometryCollectionAlgo.h"
+
+namespace Dataflow
+{
+	// Helper to get the boundary of a tet mesh, useful for debugging / verifying output
+	TArray<FIntVector3> GetSurfaceTriangles(const TArray<FIntVector4>& Tets, const bool bKeepInterior)
+	{
+		// Rotate the vector so the first element is the smallest
+		auto RotVec = [](const FIntVector3& F) -> FIntVector3
+		{
+			int32 MinIdx = F.X < F.Y ? (F.X < F.Z ? 0 : 2) : (F.Y < F.Z ? 1 : 2);
+			return FIntVector3(F[MinIdx], F[(MinIdx + 1) % 3], F[(MinIdx + 2) % 3]);
+		};
+		// Reverse the winding while keeping the first element unchanged
+		auto RevVec = [](const FIntVector3& F) -> FIntVector3
+		{
+			return FIntVector3(F.X, F.Z, F.Y);
+		};
+
+		TSet<FIntVector3> FacesSet;
+		for (int32 TetIdx = 0; TetIdx < Tets.Num(); ++TetIdx)
+		{
+			FIntVector3 TetF[4];
+			Chaos::Utilities::GetTetFaces(Tets[TetIdx], TetF[0], TetF[1], TetF[2], TetF[3], false);
+			for (int32 SubIdx = 0; SubIdx < 4; ++SubIdx)
+			{
+				// A face can be shared between a maximum of 2 tets, so no need worrying about 
+				// re-adding removed faces.
+				FIntVector3 Key = RotVec(TetF[SubIdx]);
+				if (FacesSet.Contains(Key))
+				{
+					if (!bKeepInterior)
+					{
+						FacesSet.Remove(Key);
+					}
+				}
+				else
+				{
+					FacesSet.Add(RevVec(Key));
+				}
+			}
+		}
+		return FacesSet.Array();
+	}
+
+
+
+	TArray<int32> GetMatchingMeshIndices(const TArray<FString>& MeshNames, const FManagedArrayCollection* InCollection)
+	{
+		TArray<int32> GeometryIndices;
+
+		// TransformGroup
+		int32 NumTransforms = InCollection->NumElements(FTransformCollection::TransformGroup);
+		const TManagedArray<FString>* TransformName = InCollection->FindAttribute<FString>("BoneName", FTransformCollection::TransformGroup);
+		const TManagedArray<int32>* TransformToGeometryIndex = InCollection->FindAttribute<int32>("TransformToGeometryIndex", FTransformCollection::TransformGroup);
+
+		// Geometry Group
+		int32 NumGeometry = InCollection->NumElements(FGeometryCollection::GeometryGroup);
+
+		if (MeshNames.Num())
+		{
+			for (int32 Tdx = 0; Tdx < NumTransforms; Tdx++)
+			{
+				if (MeshNames.Contains((*TransformName)[Tdx]))
+				{
+					if (0 <= (*TransformToGeometryIndex)[Tdx] && (*TransformToGeometryIndex)[Tdx] < NumGeometry)
+					{
+						GeometryIndices.Add((*TransformToGeometryIndex)[Tdx]);
+					}
+				}
+			}
+		}
+		else
+		{
+			GeometryCollectionAlgo::ContiguousArray(GeometryIndices, NumGeometry);
+		}
+
+		return GeometryIndices;
+	}
+}
\ No newline at end of file
diff --git a/Engine/Plugins/Experimental/ChaosFlesh/Source/ChaosFleshNodes/Private/Dataflow/ChaosFleshTetrahedralNodes.cpp b/Engine/Plugins/Experimental/ChaosFlesh/Source/ChaosFleshNodes/Private/Dataflow/ChaosFleshTetrahedralNodes.cpp
--- a/Engine/Plugins/Experimental/ChaosFlesh/Source/ChaosFleshNodes/Private/Dataflow/ChaosFleshTetrahedralNodes.cpp
+++ b/Engine/Plugins/Experimental/ChaosFlesh/Source/ChaosFleshNodes/Private/Dataflow/ChaosFleshTetrahedralNodes.cpp
@@ -39,48 +39,9 @@ namespace Dataflow
 		DATAFLOW_NODE_REGISTER_CREATION_FACTORY(FConstructTetGridNode);
 	}
 
-	// Helper to get the boundary of a tet mesh, useful for debugging / verifying output
-	TArray<FIntVector3> GetSurfaceTriangles(const TArray<FIntVector4>& Tets, const bool bKeepInterior)
-	{
-		// Rotate the vector so the first element is the smallest
-		auto RotVec = [](const FIntVector3& F) -> FIntVector3
-		{
-			int32 MinIdx = F.X < F.Y ? (F.X < F.Z ? 0 : 2) : (F.Y < F.Z ? 1 : 2);
-			return FIntVector3(F[MinIdx], F[(MinIdx + 1) % 3], F[(MinIdx + 2) % 3]);
-		};
-		// Reverse the winding while keeping the first element unchanged
-		auto RevVec = [](const FIntVector3& F) -> FIntVector3
-		{
-			return FIntVector3(F.X, F.Z, F.Y);
-		};
-
-		TSet<FIntVector3> FacesSet;
-		for (int32 TetIdx = 0; TetIdx < Tets.Num(); ++TetIdx)
-		{
-			FIntVector3 TetF[4];
-			Chaos::Utilities::GetTetFaces(Tets[TetIdx], TetF[0], TetF[1], TetF[2], TetF[3], false);
-			for (int32 SubIdx = 0; SubIdx < 4; ++SubIdx)
-			{
-				// A face can be shared between a maximum of 2 tets, so no need worrying about 
-				// re-adding removed faces.
-				FIntVector3 Key = RotVec(TetF[SubIdx]);
-				if (FacesSet.Contains(Key))
-				{
-					if (!bKeepInterior)
-					{
-						FacesSet.Remove(Key);
-					}
-				}
-				else
-				{
-					FacesSet.Add(RevVec(Key));
-				}
-			}
-		}
-		return FacesSet.Array();
-	}
 }
 
+
 //=============================================================================
 // FCalculateTetMetrics
 //=============================================================================
diff --git a/Engine/Plugins/Experimental/ChaosFlesh/Source/ChaosFleshNodes/Public/Dataflow/ChaosFleshCreateTetrahedronNode.h b/Engine/Plugins/Experimental/ChaosFlesh/Source/ChaosFleshNodes/Public/Dataflow/ChaosFleshCreateTetrahedronNode.h
new file mode 100644
--- /dev/null
+++ b/Engine/Plugins/Experimental/ChaosFlesh/Source/ChaosFleshNodes/Public/Dataflow/ChaosFleshCreateTetrahedronNode.h
@@ -0,0 +1,106 @@
+// Copyright Epic Games, Inc. All Rights Reserved.
+
+#pragma once 
+
+#include "CoreMinimal.h"
+#include "Dataflow/DataflowEngine.h"
+#include "GeometryCollection/ManagedArrayCollection.h"
+#include "Dataflow/ChaosFleshNodesUtility.h"
+
+#include "ChaosFleshCreateTetrahedronNode.generated.h"
+
+
+class FFleshCollection;
+namespace UE {
+	namespace Geometry {
+		class FDynamicMesh3;
+	}
+}
+
+USTRUCT(meta = (DataflowFlesh))
+struct FCreateTetrahedronDataflowNode : public FDataflowNode
+{
+	GENERATED_USTRUCT_BODY()
+	DATAFLOW_NODE_DEFINE_INTERNAL(FCreateTetrahedronDataflowNode, "CreateTetrahedron", "Flesh", "")
+	DATAFLOW_NODE_RENDER_TYPE(FGeometryCollection::StaticType(), "Collection")
+
+public:
+	typedef FManagedArrayCollection DataType;
+
+	UPROPERTY(EditAnywhere, Category = "Dataflow")
+	TEnumAsByte<TetMeshingMethod> Method = TetMeshingMethod::IsoStuffing;
+
+	UPROPERTY(EditAnywhere, Category = "Dataflow")
+	TArray<FString> MeshNames;
+
+	//
+	// IsoStuffing
+	//
+
+	UPROPERTY(EditAnywhere, Category = "Dataflow", meta = (ClampMin = "1", EditCondition = "Method == TetMeshingMethod::IsoStuffing", EditConditionHides))
+	int32 NumCells = 32;
+
+	UPROPERTY(EditAnywhere, Category = "Dataflow", meta = (ClampMin = "-0.5", ClampMax = "0.5", EditCondition = "Method == TetMeshingMethod::IsoStuffing", EditConditionHides))
+	double OffsetPercent = 0.05;
+
+	//
+	// TetWild
+	//
+
+	//! Energy at which to stop optimizing tet quality and accept the result.
+	UPROPERTY(EditAnywhere, Category = "Dataflow", meta = (ClampMin = "0.0", EditCondition = "Method == TetMeshingMethod::TetWild", EditConditionHides))
+	double IdealEdgeLength = 1.0;
+
+	//! Maximum number of optimization iterations.
+	UPROPERTY(EditAnywhere, Category = "Dataflow", meta = (ClampMin = "1", EditCondition = "Method == TetMeshingMethod::TetWild", EditConditionHides))
+	int32 MaxIterations = 80;
+
+	//! Energy at which to stop optimizing tet quality and accept the result.
+	UPROPERTY(EditAnywhere, Category = "Dataflow", meta = (ClampMin = "0.0", EditCondition = "Method == TetMeshingMethod::TetWild", EditConditionHides))
+	double StopEnergy = 10;
+
+	//! Relative tolerance, controlling how closely the mesh must follow the input surface.
+	UPROPERTY(EditAnywhere, Category = "Dataflow", meta = (ClampMin = "0.0", EditCondition = "Method == TetMeshingMethod::TetWild", EditConditionHides))
+	double EpsRel = 1e-3;
+
+	//! Coarsen the tet mesh result.
+	UPROPERTY(EditAnywhere, Category = "Dataflow", meta = (EditCondition = "Method == TetMeshingMethod::TetWild", EditConditionHides))
+	bool bCoarsen = false;
+
+	//! Enforce that the output boundary surface should be manifold.
+	UPROPERTY(EditAnywhere, Category = "Dataflow", meta = (EditCondition = "Method == TetMeshingMethod::TetWild", EditConditionHides))
+	bool bExtractManifoldBoundarySurface = false;
+
+	//! Skip the initial simplification step.
+	UPROPERTY(EditAnywhere, Category = "Dataflow", meta = (EditCondition = "Method == TetMeshingMethod::TetWild", EditConditionHides))
+	bool bSkipSimplification = false;
+
+	//! Invert tetrahedra.
+	UPROPERTY(EditAnywhere, Category = "Dataflow", meta = (EditCondition = "Method == TetMeshingMethod::TetWild", EditConditionHides))
+	bool bInvertOutputTets = false;
+
+	//
+	// Common
+	//
+	UPROPERTY(EditAnywhere, Category = "Dataflow")
+	bool bDiscardInteriorTriangles = true;
+
+	UPROPERTY(meta = (DataflowInput, DataflowOutput, DisplayName = "Collection"))
+	FManagedArrayCollection Collection;
+
+	FCreateTetrahedronDataflowNode(const Dataflow::FNodeParameters& InParam, FGuid InGuid = FGuid::NewGuid())
+		: FDataflowNode(InParam, InGuid)
+	{
+		RegisterInputConnection(&Collection);
+		RegisterOutputConnection(&Collection);
+	}
+
+	virtual void Evaluate(Dataflow::FContext& Context, const FDataflowOutput* Out) const override;
+
+protected:
+	void EvaluateIsoStuffing(Dataflow::FContext& Context, FFleshCollection& InCollection, const UE::Geometry::FDynamicMesh3& DynamicMesh) const;
+	void EvaluateTetWild(Dataflow::FContext& Context, FFleshCollection& InCollection, const UE::Geometry::FDynamicMesh3& DynamicMesh) const;
+};
+
+
+
diff --git a/Engine/Plugins/Experimental/ChaosFlesh/Source/ChaosFleshNodes/Public/Dataflow/ChaosFleshNodesUtility.h b/Engine/Plugins/Experimental/ChaosFlesh/Source/ChaosFleshNodes/Public/Dataflow/ChaosFleshNodesUtility.h
new file mode 100644
--- /dev/null
+++ b/Engine/Plugins/Experimental/ChaosFlesh/Source/ChaosFleshNodes/Public/Dataflow/ChaosFleshNodesUtility.h
@@ -0,0 +1,30 @@
+// Copyright Epic Games, Inc. All Rights Reserved.
+
+#pragma once
+
+#include "CoreMinimal.h"
+
+#include "ChaosFleshNodesUtility.generated.h"
+
+struct FManagedArrayCollection;
+
+UENUM()
+enum TetMeshingMethod : int
+{
+	IsoStuffing		UMETA(DisplayName = "IsoStuffing"),
+	TetWild			UMETA(DisplayName = "TetWild"),
+};
+
+namespace Dataflow
+{
+	TArray<FIntVector3> GetSurfaceTriangles(const TArray<FIntVector4>& Tets, const bool bKeepInterior);
+
+	/**
+	*  GetMatchingMeshIndices
+	*	Find matching geometry names in the collection. If MeshNames is empty return all the geometry indices. 
+	*	@param MeshNames	: List of name to match, or empty for all indices
+	*	@param Collection	: Collection to search
+	*/
+	TArray<int32> GetMatchingMeshIndices(const TArray<FString>& MeshNames, const FManagedArrayCollection* InCollection);
+
+}
diff --git a/Engine/Plugins/Experimental/ChaosFlesh/Source/ChaosFleshNodes/Public/Dataflow/ChaosFleshTetrahedralNodes.h b/Engine/Plugins/Experimental/ChaosFlesh/Source/ChaosFleshNodes/Public/Dataflow/ChaosFleshTetrahedralNodes.h
--- a/Engine/Plugins/Experimental/ChaosFlesh/Source/ChaosFleshNodes/Public/Dataflow/ChaosFleshTetrahedralNodes.h
+++ b/Engine/Plugins/Experimental/ChaosFlesh/Source/ChaosFleshNodes/Public/Dataflow/ChaosFleshTetrahedralNodes.h
@@ -5,6 +5,7 @@
 #include "CoreMinimal.h"
 #include "Dataflow/DataflowEngine.h"
 #include "GeometryCollection/ManagedArrayCollection.h"
+#include "Dataflow/ChaosFleshNodesUtility.h"
 
 #include "ChaosFleshTetrahedralNodes.generated.h"
 
@@ -74,12 +75,6 @@ public:
 };
 
 
-UENUM()
-enum TetMeshingMethod : int
-{
-	IsoStuffing		UMETA(DisplayName = "IsoStuffing"),
-	TetWild			UMETA(DisplayName = "TetWild"),
-};
 
 USTRUCT(meta = (DataflowFlesh))
 struct FGenerateTetrahedralCollectionDataflowNodes : public FDataflowNode
@@ -177,7 +172,6 @@ protected:
 
 namespace Dataflow
 {
-	TArray<FIntVector3> GetSurfaceTriangles(const TArray<FIntVector4>& Tets, const bool bKeepInterior);
 	void ChaosFleshTetrahedralNodes();
 
 
