GeometryCollection : AppendSkeletalMesh updates

diff --git a/Engine/Plugins/Experimental/ChaosFlesh/Source/ChaosFleshNodes/Private/Dataflow/ChaosFleshCreateTetrahedronNode.cpp b/Engine/Plugins/Experimental/ChaosFlesh/Source/ChaosFleshNodes/Private/Dataflow/ChaosFleshCreateTetrahedronNode.cpp
--- a/Engine/Plugins/Experimental/ChaosFlesh/Source/ChaosFleshNodes/Private/Dataflow/ChaosFleshCreateTetrahedronNode.cpp
+++ b/Engine/Plugins/Experimental/ChaosFlesh/Source/ChaosFleshNodes/Private/Dataflow/ChaosFleshCreateTetrahedronNode.cpp
@@ -41,6 +41,7 @@ void FCreateTetrahedronDataflowNode::Evaluate(Dataflow::FContext& Context, const
 {
 	if (Out->IsA<DataType>(&Collection))
 	{
+		FFleshCollection GeneratedTetrahedron;
 		TUniquePtr<FFleshCollection> InCollection(GetValue<DataType>(Context, &Collection).NewCopy<FFleshCollection>());
 
 		// TransformGroup
@@ -49,7 +50,7 @@ void FCreateTetrahedronDataflowNode::Evaluate(Dataflow::FContext& Context, const
 		TManagedArray<int32>* TransformToGeometryIndex = InCollection->FindAttribute<int32>("TransformToGeometryIndex", FTransformCollection::TransformGroup);
 		// Geometry Group
 		int32 NumGeometry = InCollection->NumElements(FGeometryCollection::GeometryGroup);
-		TManagedArray<int32>* GroupTransformIndex = InCollection->FindAttribute<int32>("TransformIndex", FGeometryCollection::GeometryGroup);
+		TManagedArray<int32>* GroupToTransformIndex = InCollection->FindAttribute<int32>("TransformIndex", FGeometryCollection::GeometryGroup);
 		TManagedArray<int32>* VertexCount = InCollection->FindAttribute<int32>("VertexCount", FGeometryCollection::GeometryGroup);
 		TManagedArray<int32>* VertexStart = InCollection->FindAttribute<int32>("VertexStart", FGeometryCollection::GeometryGroup);
 		TManagedArray<int32>* FaceCount = InCollection->FindAttribute<int32>("FaceCount", FGeometryCollection::GeometryGroup);
@@ -61,7 +62,7 @@ void FCreateTetrahedronDataflowNode::Evaluate(Dataflow::FContext& Context, const
 		int32 NumTriangles = InCollection->NumElements(FGeometryCollection::FacesGroup);
 		const TManagedArray<FIntVector>* Faces = InCollection->FindAttribute<FIntVector>("Indices", FGeometryCollection::FacesGroup);
 
-		if (TransformName && TransformToGeometryIndex && GroupTransformIndex && VertexCount && VertexStart && FaceCount && FaceStart
+		if (TransformName && TransformToGeometryIndex && GroupToTransformIndex && VertexCount && VertexStart && FaceCount && FaceStart
 			&& Vertex && Faces)
 		{
 			TArray<int32> ProcessGeometryIndices = Dataflow::GetMatchingMeshIndices(MeshNames, InCollection.Get());
@@ -104,27 +105,49 @@ void FCreateTetrahedronDataflowNode::Evaluate(Dataflow::FContext& Context, const
 					if (VertexToDeleteSet.Num()) TetCollection.RemoveElements(FGeometryCollection::VerticesGroup, SortedVertices);
 				}
 			});
-			for (int32 Gdx = 0; Gdx < NumGeometry; Gdx++)
+
+			auto AppendProcessedGeometry = [&ProcessGeometryIndices, &CollectionBuffer, &InCollection](FFleshCollection& ToCollection)
 			{
-				if (CollectionBuffer[Gdx])
+				TManagedArray<int32>* SourceGroupToTransformIndex = InCollection->FindAttribute<int32>("TransformIndex", FGeometryCollection::GeometryGroup);
+				TManagedArray<FString>* SourceTransformName = InCollection->FindAttribute<FString>("BoneName", FTransformCollection::TransformGroup);
+				TManagedArray<int32>* ToGroupToTransformIndex = ToCollection.FindAttribute<int32>("TransformIndex", FGeometryCollection::GeometryGroup);
+				TManagedArray<FString>* ToTransformName = ToCollection.FindAttribute<FString>("BoneName", FTransformCollection::TransformGroup);
+
+				if (SourceGroupToTransformIndex && SourceTransformName && ToGroupToTransformIndex && ToTransformName)
 				{
-					int32 GeomIndex = InCollection->AppendGeometry(*CollectionBuffer[Gdx]);
-					FString ParentName = FString::Printf(TEXT("%d"), GeomIndex);
-					if (0 <= (*GroupTransformIndex)[Gdx]
-						&& (*GroupTransformIndex)[Gdx] < InCollection->NumElements(FGeometryCollection::GeometryGroup)
-						&& !(*TransformName)[(*GroupTransformIndex)[Gdx]].IsEmpty())
+					for (int32 Sdx = 0; Sdx < ProcessGeometryIndices.Num(); Sdx++)
 					{
-						ParentName = (*TransformName)[(*GroupTransformIndex)[Gdx]];
-					}
-					if (0 <= (*GroupTransformIndex)[GeomIndex]
-						&& (*GroupTransformIndex)[GeomIndex] < InCollection->NumElements(FGeometryCollection::GeometryGroup))
-					{
-						(*TransformName)[(*GroupTransformIndex)[GeomIndex]] = FString::Printf(TEXT("%s_TET"), *ParentName);
+						int32 Gdx = ProcessGeometryIndices[Sdx];
+						if (CollectionBuffer[Gdx])
+						{
+							int32 GeomIndex = ToCollection.NumElements(FGeometryCollection::GeometryGroup);
+							ToCollection.AppendGeometry(*CollectionBuffer[Gdx]);
+							// source data
+							int32 SourceNumTransforms = InCollection->NumElements(FGeometryCollection::TransformGroup);
+							int32 SourceTransformIndex = (*SourceGroupToTransformIndex)[Gdx];
+							// target data
+							int32 ToNumTransforms = ToCollection.NumElements(FGeometryCollection::TransformGroup);
+							int32 ToGeomTransformIndex = (*ToGroupToTransformIndex)[GeomIndex];
+
+							FString TetName = FString::Printf(TEXT("%d_TET"), GeomIndex);
+							if (0 <= SourceTransformIndex && SourceTransformIndex < SourceNumTransforms)
+							{
+								if (!(*SourceTransformName)[SourceTransformIndex].IsEmpty())
+								{
+									TetName = FString::Printf(TEXT("%s_%s"), *(*SourceTransformName)[SourceTransformIndex], *TetName);
+								}
+							}
+							if (0 <= ToGeomTransformIndex && ToGeomTransformIndex < ToNumTransforms)
+							{
+								(*ToTransformName)[ToGeomTransformIndex] = TetName;
+							}
+						}
 					}
 				}
-			}
+			};
+			AppendProcessedGeometry(GeneratedTetrahedron);
 		}
-		SetValue<const DataType&>(Context, *InCollection, &Collection);
+		SetValue<const DataType&>(Context, MoveTemp(GeneratedTetrahedron), &Collection);
 	}
 }
 
diff --git a/Engine/Plugins/Experimental/GeometryCollectionPlugin/Source/GeometryCollectionNodes/Private/Dataflow/GeometryCollectionSkeletalMeshToCollectionNode.cpp b/Engine/Plugins/Experimental/GeometryCollectionPlugin/Source/GeometryCollectionNodes/Private/Dataflow/GeometryCollectionSkeletalMeshToCollectionNode.cpp
--- a/Engine/Plugins/Experimental/GeometryCollectionPlugin/Source/GeometryCollectionNodes/Private/Dataflow/GeometryCollectionSkeletalMeshToCollectionNode.cpp
+++ b/Engine/Plugins/Experimental/GeometryCollectionPlugin/Source/GeometryCollectionNodes/Private/Dataflow/GeometryCollectionSkeletalMeshToCollectionNode.cpp
@@ -25,110 +25,9 @@ void FSkeletalMeshToCollectionDataflowNode::Evaluate(Dataflow::FContext& Context
 		TObjectPtr<const USkeletalMesh> InSkeletalMesh = GetValue<TObjectPtr<const USkeletalMesh>>(Context, &SkeletalMesh);
 		if (InSkeletalMesh)
 		{
-			if (Method == EDF_ConversionMethod::Skinned)
-			{
-				FGeometryCollectionEngineConversion::AppendSkeletalMesh(InSkeletalMesh, 0, FTransform::Identity, &OutCollection);
-			}
-			else if (Method == EDF_ConversionMethod::StaticComponents)
-			{
-				AppendSkeletalMeshComponentsToGeometryCollection(InSkeletalMesh, OutCollection);
-			}
+			FGeometryCollectionEngineConversion::AppendSkeletalMesh(InSkeletalMesh, 0, FTransform::Identity, &OutCollection);
 		}
 		SetValue(Context, FManagedArrayCollection(OutCollection), &Collection);
 	}
 }
 
-
-
-
-void FSkeletalMeshToCollectionDataflowNode::AppendSkeletalMeshComponentsToGeometryCollection(const USkeletalMesh* InSkeletalMesh, FGeometryCollection& OutCollection) const
-{
-#if WITH_EDITORONLY_DATA
-	UE::Geometry::FDynamicMesh3 DynamicMesh;
-
-	// Check first if we have bulk data available and non-empty.
-	constexpr int32 LODIndex = 0;
-	FMeshDescription SourceMesh;
-	if (InSkeletalMesh->HasMeshDescription(LODIndex))
-	{
-		InSkeletalMesh->CloneMeshDescription(LODIndex, SourceMesh);
-	}
-	FMeshDescriptionToDynamicMesh Converter;
-	Converter.Convert(&SourceMesh, DynamicMesh);
-
-	//
-	// Compute by Component
-	//
-	bool bComputeByComponent = true;
-	if (bComputeByComponent)
-	{
-		TArray<TArray<int32>> ConnectedComponents;
-
-		TArray<FIntVector3> Faces;
-		Faces.SetNum(DynamicMesh.TriangleCount());
-		for (int32 i = 0; i < DynamicMesh.TriangleCount(); ++i)
-		{
-			Faces[i] = FIntVector3(DynamicMesh.GetTriangle(i));
-		}
-		Chaos::Utilities::FindConnectedRegions(Faces, ConnectedComponents);
-
-		TArray<TUniquePtr<FGeometryCollection>> CollectionBuffer;
-		for (int32 i = 0; i < ConnectedComponents.Num(); i++)
-		{
-			CollectionBuffer.Add(nullptr);
-		}
-		ParallelFor(ConnectedComponents.Num(),[&](int32 i)
-		{
-			UE::Geometry::FDynamicMesh3 Component;
-
-			for (FVector V : DynamicMesh.VerticesItr()) {
-				Component.AppendVertex(V);
-			}
-			for (int32 j = 0; j < ConnectedComponents[i].Num(); j++) {
-				int32 ElementIndex = ConnectedComponents[i][j];
-				for (int32 ie = 0; ie < 3; ie++) {
-					Component.AppendTriangle(Faces[ElementIndex][0], Faces[ElementIndex][1], Faces[ElementIndex][2]);
-				}
-			}
-			Component.CompactInPlace();
-			Component.ReverseOrientation();
-
-			int Idx = 0;
-			TArray<float> VertexArray;
-			VertexArray.SetNumUninitialized(Component.VertexCount() * 3);
-			for (FVector V : Component.VerticesItr())
-			{
-				VertexArray[Idx++] = V[0];
-				VertexArray[Idx++] = V[1];
-				VertexArray[Idx++] = V[2];
-			}
-
-			Idx = 0;
-			TArray<int> IndicesArray;
-			IndicesArray.SetNumUninitialized(Component.TriangleCount() * 3);
-			for (UE::Geometry::FIndex3i Tri : Component.TrianglesItr())
-			{
-				IndicesArray[Idx++] = Tri[0];
-				IndicesArray[Idx++] = Tri[1];
-				IndicesArray[Idx++] = Tri[2];
-			}
-
-			CollectionBuffer[i] = TUniquePtr<FGeometryCollection>(FGeometryCollection::NewGeometryCollection(VertexArray, IndicesArray));
-		});
-
-		for (int32 i = 0; i < CollectionBuffer.Num(); i++)
-		{
-			if (CollectionBuffer[i]->NumElements(FGeometryCollection::VerticesGroup))
-			{
-				TSet<int32> VertexToDeleteSet;
-				GeometryCollectionAlgo::ComputeStaleVertices(CollectionBuffer[i].Get(), VertexToDeleteSet);
-				TArray<int32> SortedVertices = VertexToDeleteSet.Array(); SortedVertices.Sort();
-				if (VertexToDeleteSet.Num()) CollectionBuffer[i]->RemoveElements(FGeometryCollection::VerticesGroup, SortedVertices);
-
-				int32 GeomIndex = OutCollection.AppendGeometry(*CollectionBuffer[i].Get());
-			}
-		}
-	}
-#endif
-}
-
diff --git a/Engine/Plugins/Experimental/GeometryCollectionPlugin/Source/GeometryCollectionNodes/Public/Dataflow/GeometryCollectionSkeletalMeshToCollectionNode.h b/Engine/Plugins/Experimental/GeometryCollectionPlugin/Source/GeometryCollectionNodes/Public/Dataflow/GeometryCollectionSkeletalMeshToCollectionNode.h
--- a/Engine/Plugins/Experimental/GeometryCollectionPlugin/Source/GeometryCollectionNodes/Public/Dataflow/GeometryCollectionSkeletalMeshToCollectionNode.h
+++ b/Engine/Plugins/Experimental/GeometryCollectionPlugin/Source/GeometryCollectionNodes/Public/Dataflow/GeometryCollectionSkeletalMeshToCollectionNode.h
@@ -10,12 +10,6 @@
 
 class USkeletalMesh;
 
-UENUM()
-enum EDF_ConversionMethod : int
-{
-	Skinned				UMETA(DisplayName = "Skinned"),
-	StaticComponents	UMETA(DisplayName = "Static Components"),
-};
 
 USTRUCT(meta = (DataflowGeometryCollection))
 struct FSkeletalMeshToCollectionDataflowNode : public FDataflowNode
@@ -27,9 +21,6 @@ struct FSkeletalMeshToCollectionDataflowNode : public FDataflowNode
 public:
 	typedef FManagedArrayCollection DataType;
 
-	UPROPERTY(EditAnywhere, Category = "Dataflow")
-	TEnumAsByte<EDF_ConversionMethod> Method = EDF_ConversionMethod::StaticComponents;
-
 	UPROPERTY(EditAnywhere, Category = "Dataflow", meta = (DataflowInput, DisplayName = "SkeletalMesh"))
 	TObjectPtr<const USkeletalMesh> SkeletalMesh = nullptr;
 
@@ -44,9 +35,5 @@ public:
 	}
 
 	virtual void Evaluate(Dataflow::FContext& Context, const FDataflowOutput* Out) const override;
-
-private:
-	void AppendSkeletalMeshComponentsToGeometryCollection(const USkeletalMesh* InSkeletalMesh, FGeometryCollection& OutCollection) const;
-
 };
 
diff --git a/Engine/Source/Runtime/Experimental/GeometryCollectionEngine/Private/GeometryCollection/GeometryCollectionEngineConversion.cpp b/Engine/Source/Runtime/Experimental/GeometryCollectionEngine/Private/GeometryCollection/GeometryCollectionEngineConversion.cpp
--- a/Engine/Source/Runtime/Experimental/GeometryCollectionEngine/Private/GeometryCollection/GeometryCollectionEngineConversion.cpp
+++ b/Engine/Source/Runtime/Experimental/GeometryCollectionEngine/Private/GeometryCollection/GeometryCollectionEngineConversion.cpp
@@ -1329,13 +1329,19 @@ bool FGeometryCollectionEngineConversion::AppendSkeletalMesh(const USkeletalMesh
 							TargetVertexVisited[TargetVertexIndex] = true;
 
 							//colors
-							Color[TargetVertexIndex] = FLinearColor::White;
+							Color[TargetVertexIndex] = DefaultColor;
 							if (ColorOverlay)
 							{
-								FVertexInstanceID SourceInstanceID = MeshDescription.GetTriangleVertexInstance(SourceTriangleIndex, k);
-								FVector4f InstColor = InstanceColors.Get(SourceInstanceID);
-								UE::Geometry::LinearColors::SRGBToLinear(InstColor);
-								//Color[TargetVertexIndex] = InstColor;
+								// k : 0 -> 2
+								// SourceTriangleIndex is a TriangleID from the DynamicMesh
+								// TargetVertexIndex is the index for the vertex in the GeometryCollection
+								int32 SourceInstanceID = ColorOverlay->GetTriangle(SourceTriangleIndex)[k];
+								if (ColorOverlay->IsElement(SourceInstanceID))
+								{
+									FVector4f InstColor = ColorOverlay->GetElement(SourceInstanceID);
+									UE::Geometry::LinearColors::SRGBToLinear(InstColor);
+									//Color[TargetVertexIndex] = FLinearColor(InstColor);
+								}
 							}
 
 							// @todo(GeometryCollectionConversion) : Add support for UV's, Normals  
