ChaosFlesh : Dataflow node rendering update

diff --git a/Engine/Plugins/Experimental/ChaosFlesh/Source/ChaosFleshNodes/Public/Dataflow/ChaosFleshCoreNodes.h b/Engine/Plugins/Experimental/ChaosFlesh/Source/ChaosFleshNodes/Public/Dataflow/ChaosFleshCoreNodes.h
--- a/Engine/Plugins/Experimental/ChaosFlesh/Source/ChaosFleshNodes/Public/Dataflow/ChaosFleshCoreNodes.h
+++ b/Engine/Plugins/Experimental/ChaosFlesh/Source/ChaosFleshNodes/Public/Dataflow/ChaosFleshCoreNodes.h
@@ -17,7 +17,7 @@ struct FAppendTetrahedralCollectionDataflowNode : public FDataflowNode
 {
 	GENERATED_USTRUCT_BODY()
 	DATAFLOW_NODE_DEFINE_INTERNAL(FAppendTetrahedralCollectionDataflowNode, "AppendTetrahedralCollection", "Flesh", "")
-	DATAFLOW_NODE_RENDER_TYPE(FGeometryCollection::StaticType(), "Collection")
+	DATAFLOW_NODE_RENDER_TYPE(FGeometryCollection::StaticType(), "Collection1")
 public:
 	UPROPERTY(meta = (DataflowInput, DataflowOutput, DisplayName = "Collection1", DataflowPassthrough = "Collection1"));
 	FManagedArrayCollection Collection1;
