ChaosFlesh : Updated  transform name index

diff --git a/Engine/Plugins/Experimental/ChaosFlesh/Source/ChaosFleshNodes/Private/Dataflow/ChaosFleshCreateTetrahedronNode.cpp b/Engine/Plugins/Experimental/ChaosFlesh/Source/ChaosFleshNodes/Private/Dataflow/ChaosFleshCreateTetrahedronNode.cpp
--- a/Engine/Plugins/Experimental/ChaosFlesh/Source/ChaosFleshNodes/Private/Dataflow/ChaosFleshCreateTetrahedronNode.cpp
+++ b/Engine/Plugins/Experimental/ChaosFlesh/Source/ChaosFleshNodes/Private/Dataflow/ChaosFleshCreateTetrahedronNode.cpp
@@ -129,7 +129,7 @@ void FCreateTetrahedronDataflowNode::Evaluate(Dataflow::FContext& Context, const
 							int32 ToNumTransforms = ToCollection.NumElements(FGeometryCollection::TransformGroup);
 							int32 ToGeomTransformIndex = (*ToGroupToTransformIndex)[GeomIndex];
 
-							FString TetName = FString::Printf(TEXT("%d_TET"), GeomIndex);
+							FString TetName = FString::Printf(TEXT("Tet%d"), GeomIndex);
 							if (0 <= SourceTransformIndex && SourceTransformIndex < SourceNumTransforms)
 							{
 								if (!(*SourceTransformName)[SourceTransformIndex].IsEmpty())
