ChaosFlesh : Dont add empty elements.

diff --git a/Engine/Plugins/Experimental/ChaosFlesh/Source/ChaosFleshNodes/Private/Dataflow/ChaosFleshCreateTetrahedronNode.cpp b/Engine/Plugins/Experimental/ChaosFlesh/Source/ChaosFleshNodes/Private/Dataflow/ChaosFleshCreateTetrahedronNode.cpp
--- a/Engine/Plugins/Experimental/ChaosFlesh/Source/ChaosFleshNodes/Private/Dataflow/ChaosFleshCreateTetrahedronNode.cpp
+++ b/Engine/Plugins/Experimental/ChaosFlesh/Source/ChaosFleshNodes/Private/Dataflow/ChaosFleshCreateTetrahedronNode.cpp
@@ -118,7 +118,7 @@ void FCreateTetrahedronDataflowNode::Evaluate(Dataflow::FContext& Context, const
 					for (int32 Sdx = 0; Sdx < ProcessGeometryIndices.Num(); Sdx++)
 					{
 						int32 Gdx = ProcessGeometryIndices[Sdx];
-						if (CollectionBuffer[Gdx])
+						if (CollectionBuffer[Gdx] && CollectionBuffer[Gdx]->NumElements(FGeometryCollection::GeometryGroup))
 						{
 							int32 GeomIndex = ToCollection.NumElements(FGeometryCollection::GeometryGroup);
 							ToCollection.AppendGeometry(*CollectionBuffer[Gdx]);
