ChaosFlesh : Tetrahedralization
- Consider transform when generating tetrahdron

diff --git a/Engine/Plugins/Experimental/ChaosFlesh/Source/ChaosFleshNodes/Private/Dataflow/ChaosFleshCreateTetrahedronNode.cpp b/Engine/Plugins/Experimental/ChaosFlesh/Source/ChaosFleshNodes/Private/Dataflow/ChaosFleshCreateTetrahedronNode.cpp
--- a/Engine/Plugins/Experimental/ChaosFlesh/Source/ChaosFleshNodes/Private/Dataflow/ChaosFleshCreateTetrahedronNode.cpp
+++ b/Engine/Plugins/Experimental/ChaosFlesh/Source/ChaosFleshNodes/Private/Dataflow/ChaosFleshCreateTetrahedronNode.cpp
@@ -48,6 +48,8 @@ void FCreateTetrahedronDataflowNode::Evaluate(Dataflow::FContext& Context, const
 		int32 NumTransforms = InCollection->NumElements(FTransformCollection::TransformGroup);
 		TManagedArray<FString>* TransformName = InCollection->FindAttribute<FString>("BoneName", FTransformCollection::TransformGroup);
 		TManagedArray<int32>* TransformToGeometryIndex = InCollection->FindAttribute<int32>("TransformToGeometryIndex", FTransformCollection::TransformGroup);
+		const TManagedArray<int32>* Parent = InCollection->FindAttribute<int32>(FTransformCollection::ParentAttribute, FTransformCollection::TransformGroup);
+		const TManagedArray<FTransform3f>* LocalSpaceTransform = InCollection->FindAttribute<FTransform3f>(FTransformCollection::TransformAttribute, FTransformCollection::TransformGroup);
 		// Geometry Group
 		int32 NumGeometry = InCollection->NumElements(FGeometryCollection::GeometryGroup);
 		TManagedArray<int32>* GroupToTransformIndex = InCollection->FindAttribute<int32>("TransformIndex", FGeometryCollection::GeometryGroup);
@@ -55,6 +57,7 @@ void FCreateTetrahedronDataflowNode::Evaluate(Dataflow::FContext& Context, const
 		TManagedArray<int32>* VertexStart = InCollection->FindAttribute<int32>("VertexStart", FGeometryCollection::GeometryGroup);
 		TManagedArray<int32>* FaceCount = InCollection->FindAttribute<int32>("FaceCount", FGeometryCollection::GeometryGroup);
 		TManagedArray<int32>* FaceStart = InCollection->FindAttribute<int32>("FaceStart", FGeometryCollection::GeometryGroup);
+
 		// Vertices Group
 		int32 NumVertices = InCollection->NumElements(FGeometryCollection::VerticesGroup);
 		const TManagedArray<FVector3f>* Vertex = InCollection->FindAttribute<FVector3f>("Vertex", FGeometryCollection::VerticesGroup);
@@ -62,7 +65,8 @@ void FCreateTetrahedronDataflowNode::Evaluate(Dataflow::FContext& Context, const
 		int32 NumTriangles = InCollection->NumElements(FGeometryCollection::FacesGroup);
 		const TManagedArray<FIntVector>* Faces = InCollection->FindAttribute<FIntVector>("Indices", FGeometryCollection::FacesGroup);
 
-		if (TransformName && TransformToGeometryIndex && GroupToTransformIndex && VertexCount && VertexStart && FaceCount && FaceStart
+		if (TransformName && TransformToGeometryIndex && Parent && LocalSpaceTransform &&
+			GroupToTransformIndex && VertexCount && VertexStart && FaceCount && FaceStart
 			&& Vertex && Faces)
 		{
 			TArray<int32> ProcessGeometryIndices = Dataflow::GetMatchingMeshIndices(MeshNames, InCollection.Get());
@@ -76,14 +80,18 @@ void FCreateTetrahedronDataflowNode::Evaluate(Dataflow::FContext& Context, const
 					CollectionBuffer.Add(nullptr);
 			}
 
+			TArray<FTransform> ComponentTransform;
+			GeometryCollectionAlgo::GlobalMatrices(*LocalSpaceTransform, *Parent, ComponentTransform);
+
 			ParallelFor(ProcessGeometryIndices.Num(), [&](int32 i)
 			{
 				int32 Gdx = ProcessGeometryIndices[i];
+			    int32 Tdx = (*GroupToTransformIndex)[Gdx];
 				FFleshCollection& TetCollection = *CollectionBuffer[Gdx];
 
 				UE::Geometry::FDynamicMesh3 DynamicMesh;
 				int32 vStart = (*VertexStart)[Gdx], vEnd = vStart + (*VertexCount)[Gdx];
-				for (int Vdx = vStart; Vdx < vEnd; Vdx++) DynamicMesh.AppendVertex(FVector((*Vertex)[Vdx]));
+				for (int Vdx = vStart; Vdx < vEnd; Vdx++) DynamicMesh.AppendVertex(ComponentTransform[Tdx].TransformPosition(FVector((*Vertex)[Vdx])));
 				int32 fStart = (*FaceStart)[Gdx], fEnd = fStart + (*FaceCount)[Gdx];
 				for (int Fdx = fStart; Fdx < fEnd; Fdx++) DynamicMesh.AppendTriangle((*Faces)[Fdx] - FIntVector(vStart));
 				DynamicMesh.CompactInPlace();
