ChaosFlesh: Deprecate ComponentIndex and replace with GeometryGroup

diff --git a/Engine/Plugins/Experimental/ChaosFlesh/Source/ChaosFleshNodes/Private/Dataflow/ChaosFleshEngineAssetNodes.cpp b/Engine/Plugins/Experimental/ChaosFlesh/Source/ChaosFleshNodes/Private/Dataflow/ChaosFleshEngineAssetNodes.cpp
--- a/Engine/Plugins/Experimental/ChaosFlesh/Source/ChaosFleshNodes/Private/Dataflow/ChaosFleshEngineAssetNodes.cpp
+++ b/Engine/Plugins/Experimental/ChaosFlesh/Source/ChaosFleshNodes/Private/Dataflow/ChaosFleshEngineAssetNodes.cpp
@@ -13,6 +13,7 @@
 #include "Dataflow/DataflowInputOutput.h"
 #include "Dataflow/DataflowObjectInterface.h"
 #include "Dataflow/DataflowNodeFactory.h"
+#include "GeometryCollection/Facades/CollectionMeshFacade.h"
 #include "GeometryCollection/Facades/CollectionMuscleActivationFacade.h"
 #include "GeometryCollection/TransformCollection.h"
 #include UE_INLINE_GENERATED_CPP_BY_NAME(ChaosFleshEngineAssetNodes)
@@ -447,48 +448,37 @@ void FComputeMuscleActivationDataNode::Evaluate(Dataflow::FContext& Context, con
 						TArray<TArray<int32>> MuscleActivationElements;
 						TArray<int32> OriginNodes;
 						TArray<int32> InsertionNodes;
-						if (TManagedArray<int32>* ComponentIndex = InCollection.FindAttribute<int32>("ComponentIndex", FGeometryCollection::VerticesGroup))
+						GeometryCollection::Facades::FCollectionMeshFacade MeshFacade(InCollection);
+						TArray<int32> ComponentIndex = MeshFacade.GetGeometryGroupIndexArray();
+						TMap<int32, int32> ComponentToIndex; //Component index to muscle index
+						for (int32 i = 0; i < InOriginIndices.Num(); i++)
 						{
-							TMap<int32, int32> ComponentToIndex; //Component index to muscle index
-							for (int32 i = 0; i < InOriginIndices.Num(); i++)
+							if (!ComponentToIndex.Contains(ComponentIndex[InOriginIndices[i]]))
 							{
-								if (!ComponentToIndex.Contains((*ComponentIndex)[InOriginIndices[i]]))
-								{
-									//TODO: use one pair of separately defined origin and insertion for each muscle instead of choosing the first from kinematic origins and last from insertions
-									ComponentToIndex.Add((*ComponentIndex)[InOriginIndices[i]], OriginNodes.Num());
-									OriginNodes.Add(InOriginIndices[i]);
-								}
+								//TODO: use one pair of separately defined origin and insertion for each muscle instead of choosing the first from kinematic origins and last from insertions
+								ComponentToIndex.Add(ComponentIndex[InOriginIndices[i]], OriginNodes.Num());
+								OriginNodes.Add(InOriginIndices[i]);
 							}
-							InsertionNodes.Init(INDEX_NONE, OriginNodes.Num());
-							for (int32 i = 0; i < InInsertionIndices.Num(); i++)
+						}
+						InsertionNodes.Init(INDEX_NONE, OriginNodes.Num());
+						for (int32 i = 0; i < InInsertionIndices.Num(); i++)
+						{
+							if (!ComponentToIndex.Contains(ComponentIndex[InInsertionIndices[i]]))
 							{
-								if (!ComponentToIndex.Contains((*ComponentIndex)[InInsertionIndices[i]]))
-								{
-									ensureMsgf(false, TEXT("No origin in this component"));
-								}
-								else if (InsertionNodes[ComponentToIndex[(*ComponentIndex)[InInsertionIndices[i]]]] == INDEX_NONE)
-								{
-									InsertionNodes[ComponentToIndex[(*ComponentIndex)[InInsertionIndices[i]]]] = InInsertionIndices[i];
-								}
+								ensureMsgf(false, TEXT("No origin in this component"));
 							}
-							MuscleActivationElements.SetNum(OriginNodes.Num());
-							for (int32 ElemIdx = 0; ElemIdx < Elements->Num(); ElemIdx++)
+							else if (InsertionNodes[ComponentToIndex[ComponentIndex[InInsertionIndices[i]]]] == INDEX_NONE)
 							{
-								if (ComponentToIndex.Contains((*ComponentIndex)[(*Elements)[ElemIdx][0]]))
-								{
-									MuscleActivationElements[ComponentToIndex[(*ComponentIndex)[(*Elements)[ElemIdx][0]]]].Add(ElemIdx);
-								}
+								InsertionNodes[ComponentToIndex[ComponentIndex[InInsertionIndices[i]]]] = InInsertionIndices[i];
 							}
 						}
-						else
+						MuscleActivationElements.SetNum(OriginNodes.Num());
+						for (int32 ElemIdx = 0; ElemIdx < Elements->Num(); ElemIdx++)
 						{
-							MuscleActivationElements.SetNum(1);
-							for (int32 i = 0; i < Elements->Num(); i++)
+							if (ComponentToIndex.Contains(ComponentIndex[(*Elements)[ElemIdx][0]]))
 							{
-								MuscleActivationElements[0].Add(i);
+								MuscleActivationElements[ComponentToIndex[ComponentIndex[(*Elements)[ElemIdx][0]]]].Add(ElemIdx);
 							}
-							OriginNodes.Add(InOriginIndices[0]);
-							InsertionNodes.Add(InInsertionIndices[0]);
 						}
 						OriginInsertionRestLength.SetNum(OriginNodes.Num());
 
@@ -540,8 +530,7 @@ void FComputeIslandsNode::Evaluate(Dataflow::FContext& Context, const FDataflowO
 			FTetrahedralCollection::TetrahedronAttribute, FTetrahedralCollection::TetrahedralGroup);
 
 
-		if (InCollection.HasAttributes({
-			MType<int32>("ComponentIndex", FGeometryCollection::VerticesGroup) }) && Elements)
+		if (Elements)
 		{
 
 			int32 VertsNum = InCollection.NumElements(FGeometryCollection::VerticesGroup);
@@ -640,55 +629,50 @@ void FGenerateOriginInsertionNode::Evaluate(Dataflow::FContext& Context, const F
 		// these via an input on the node...
 		//
 		auto DoubleVert = [](FVector3f V) { return FVector3d(V.X, V.Y, V.Z); };
-		if (TManagedArray<int32>* ComponentIndex = InCollection.FindAttribute<int32>("ComponentIndex", FGeometryCollection::VerticesGroup))
+		GeometryCollection::Facades::FCollectionMeshFacade MeshFacade(InCollection);
+		TArray<int32> ComponentIndex = MeshFacade.GetGeometryGroupIndexArray();
+		// Origin vertices
+		if (!InOriginIndices.IsEmpty())
 		{
-			// Origin vertices
-			if (!InOriginIndices.IsEmpty())
+			for (int32 i = 0; i < InOriginIndices.Num(); ++i)
 			{
-				for (int32 i = 0; i < InOriginIndices.Num(); ++i)
+				if (InOriginIndices[i] < Vertex->Num())
 				{
-					if (InOriginIndices[i] < Vertex->Num())
+					for (int32 j = 0; j < Vertex->Num(); ++j)
 					{
-						for (int32 j = 0; j < Vertex->Num(); ++j)
+						if (ComponentIndex[InOriginIndices[i]] == ComponentIndex[j] 
+							&& ComponentIndex[InOriginIndices[i]] >= 0
+							&& ComponentIndex[j] >= 0
+							&& ((*Vertex)[InOriginIndices[i]] - (*Vertex)[j]).Size() < Radius)
 						{
-							if ((*ComponentIndex)[InOriginIndices[i]] == (*ComponentIndex)[j] 
-								&& (*ComponentIndex)[InOriginIndices[i]] >= 0
-								&& (*ComponentIndex)[j] >= 0
-								&& ((*Vertex)[InOriginIndices[i]] - (*Vertex)[j]).Size() < Radius)
-							{
-								OutOriginIndices.Add(j);
-							}
+							OutOriginIndices.Add(j);
 						}
 					}
 				}
 			}
+		}
 
-			// Insertion vertices
-			if (!InInsertionIndices.IsEmpty())
+		// Insertion vertices
+		if (!InInsertionIndices.IsEmpty())
+		{
+			for (int32 i = 0; i < InInsertionIndices.Num(); ++i)
 			{
-				for (int32 i = 0; i < InInsertionIndices.Num(); ++i)
+				if (InInsertionIndices[i] < Vertex->Num())
 				{
-					if (InInsertionIndices[i] < Vertex->Num())
+					for (int32 j = 0; j < Vertex->Num(); ++j)
 					{
-						for (int32 j = 0; j < Vertex->Num(); ++j)
+						if (ComponentIndex[InInsertionIndices[i]] == ComponentIndex[j]
+							&& ComponentIndex[InInsertionIndices[i]] >= 0
+							&& ComponentIndex[j] >= 0
+							&& ((*Vertex)[InInsertionIndices[i]] - (*Vertex)[j]).Size() < Radius)
 						{
-							if ((*ComponentIndex)[InInsertionIndices[i]] == (*ComponentIndex)[j]
-								&& (*ComponentIndex)[InInsertionIndices[i]] >= 0
-								&& (*ComponentIndex)[j] >= 0
-								&& ((*Vertex)[InInsertionIndices[i]] - (*Vertex)[j]).Size() < Radius)
-							{
-								OutInsertionIndices.Add(j);
-							}
+							OutInsertionIndices.Add(j);
 						}
 					}
 				}
 			}
 		}
 
-		//
-		// Set output(s)
-		//
-
 		SetValue(Context, MoveTemp(InCollection), &Collection);
 		SetValue(Context, MoveTemp(OutOriginIndices), &OriginIndicesOut);
 		SetValue(Context, MoveTemp(OutInsertionIndices), &InsertionIndicesOut);
@@ -701,37 +685,35 @@ void FIsolateComponentNode::Evaluate(Dataflow::FContext& Context, const FDataflo
 	{	
 		FManagedArrayCollection InCollection = GetValue<FManagedArrayCollection>(Context, &Collection);
 		TArray<int32> DeleteList;
-		if (TManagedArray<FIntVector>* Indices = InCollection.FindAttribute<FIntVector>("Indices", FGeometryCollection::FacesGroup))
+		TManagedArray<FIntVector>* Indices = InCollection.FindAttribute<FIntVector>("Indices", FGeometryCollection::FacesGroup);
+		TManagedArray<bool>*FaceVisibility = InCollection.FindAttribute<bool>("Visible", FGeometryCollection::FacesGroup);
+		TManagedArray<int32>* FaceStart = InCollection.FindAttribute<int32>("FaceStart", FGeometryCollection::GeometryGroup);
+		TManagedArray<int32>* FaceCount = InCollection.FindAttribute<int32>("FaceCount", FGeometryCollection::GeometryGroup);
+		if (Indices && FaceVisibility && FaceStart && FaceCount)
 		{
-			if (TManagedArray<bool>* FaceVisibility = InCollection.FindAttribute<bool>("Visible", FGeometryCollection::FacesGroup))
+			FaceVisibility->Fill(false);
+			TSet<int32> GeometrySet;
+			TArray<FString> StrArray;
+			TargetGeometryIndex.ParseIntoArray(StrArray, *FString(" "));
+			for (FString GeometryIdx : StrArray)
 			{
-				if (TManagedArray<int32>* ComponentIndex = InCollection.FindAttribute<int32>("ComponentIndex", FGeometryCollection::VerticesGroup))
-				{	
-					TSet<int32> ComponentSet;
-					TArray<FString> StrArray;
-					TargetComponentIndex.ParseIntoArray(StrArray, *FString(" "));
-					for (FString Elem : StrArray)
-					{
-						if (Elem.Len() && FCString::IsNumeric(*Elem))
-						{
-							ComponentSet.Add(FCString::Atoi(*Elem));
-						}
-					}
-					for (int32 i = 0; i < Indices->Num(); i++)
-					{
-						if (!ComponentSet.Contains((*ComponentIndex)[(*Indices)[i][0]]))
-						{
-							(*FaceVisibility)[i] = false;
-							DeleteList.AddUnique(i);
-						}
-					}
+				if (GeometryIdx.Len() && FCString::IsNumeric(*GeometryIdx))
+				{
+					GeometrySet.Add(FCString::Atoi(*GeometryIdx));
 				}
 			}
-		}
-		if (bDeleteHiddenFaces)
-		{
-			DeleteList.Sort();
-			InCollection.RemoveElements(FGeometryCollection::FacesGroup, DeleteList);
+			for (TSet<int32>::TConstIterator It = GeometrySet.CreateConstIterator(); It; ++It)
+			{
+				for (int32 FaceIdx = (*FaceStart)[*It]; FaceIdx < (*FaceStart)[*It] + (*FaceCount)[*It]; FaceIdx++)
+				{
+					(*FaceVisibility)[FaceIdx] = true;
+					DeleteList.Add(FaceIdx);
+				}
+			}
+			if (bDeleteHiddenFaces)
+			{
+				InCollection.RemoveElements(FGeometryCollection::FacesGroup, DeleteList);
+			}
 		}
 		SetValue(Context, MoveTemp(InCollection), &Collection);
 	}
diff --git a/Engine/Plugins/Experimental/ChaosFlesh/Source/ChaosFleshNodes/Private/Dataflow/ChaosFleshPositionTargetInitializationNodes.cpp b/Engine/Plugins/Experimental/ChaosFlesh/Source/ChaosFleshNodes/Private/Dataflow/ChaosFleshPositionTargetInitializationNodes.cpp
--- a/Engine/Plugins/Experimental/ChaosFlesh/Source/ChaosFleshNodes/Private/Dataflow/ChaosFleshPositionTargetInitializationNodes.cpp
+++ b/Engine/Plugins/Experimental/ChaosFlesh/Source/ChaosFleshNodes/Private/Dataflow/ChaosFleshPositionTargetInitializationNodes.cpp
@@ -9,9 +9,10 @@
 #include "Dataflow/DataflowInputOutput.h"
 #include "Chaos/BoundingVolumeHierarchy.h"
 #include "GeometryCollection/Facades/CollectionKinematicBindingFacade.h"
+#include "GeometryCollection/Facades/CollectionMeshFacade.h"
 #include "GeometryCollection/Facades/CollectionPositionTargetFacade.h"
-#include "Dataflow/DataflowNodeFactory.h"
 #include "GeometryCollection/Facades/CollectionVertexBoneWeightsFacade.h"
+#include "Dataflow/DataflowNodeFactory.h"
 #include "GeometryCollection/TransformCollection.h"
 #include "Chaos/Tetrahedron.h"
 #include "Chaos/Plane.h"
@@ -397,163 +398,161 @@ void FSetVertexTrianglePositionTargetBindingDataflowNode::Evaluate(Dataflow::FCo
 			{
 				if (TManagedArray<FIntVector>* Indices = InCollection.FindAttribute<FIntVector>("Indices", FGeometryCollection::FacesGroup))
 				{
-					if (TManagedArray<int32>* ComponentIndex = InCollection.FindAttribute<int32>("ComponentIndex", FGeometryCollection::VerticesGroup))
+					GeometryCollection::Facades::FCollectionMeshFacade MeshFacade(InCollection);
+					TArray<int32> ComponentIndex = MeshFacade.GetGeometryGroupIndexArray();
+					TArray<Chaos::TVector<int32, 3>> IndicesArray;
+					for (int32 i = 0; i < Indices->Num(); i++)
 					{
-
-						TArray<Chaos::TVector<int32, 3>> IndicesArray;
-						for (int32 i = 0; i < Indices->Num(); i++)
+						Chaos::TVector<int32, 3> CurrentIndices(0);
+						for (int32 j = 0; j < 3; j++) 
 						{
-							Chaos::TVector<int32, 3> CurrentIndices(0);
-							for (int32 j = 0; j < 3; j++) 
-							{
-								CurrentIndices[j] = (*Indices)[i][j];
-							}
-							if (CurrentIndices[0] != -1
-								&& CurrentIndices[1] != -1
-								&& CurrentIndices[2] != -1)
-							{
-								IndicesArray.Emplace(CurrentIndices);
-							}
+							CurrentIndices[j] = (*Indices)[i][j];
 						}
-						TArray<TArray<int32>> LocalIndex;
-						TArray<TArray<int32>>* LocalIndexPtr = &LocalIndex;
-						TArray<TArray<int>> GlobalIndex = Chaos::Utilities::ComputeIncidentElements(IndicesArray, LocalIndexPtr);
-						int32 ActualParticleCount = 0;
-						for (int32 l = 0; l < GlobalIndex.Num(); l++)
+						if (CurrentIndices[0] != -1
+							&& CurrentIndices[1] != -1
+							&& CurrentIndices[2] != -1)
 						{
-							if (GlobalIndex[l].Num() > 0)
-							{
-								ActualParticleCount += 1;
-							}
+							IndicesArray.Emplace(CurrentIndices);
 						}
-						TArray<Chaos::TVector<float, 3>> IndicesPositions; 
-						IndicesPositions.SetNum(ActualParticleCount);
-						TArray<int32> IndicesMap;
-						IndicesMap.SetNum(ActualParticleCount);
-						int32 CurrentParticleIndex = 0;
-						for (int32 i = 0; i < GlobalIndex.Num(); i++)
+					}
+					TArray<TArray<int32>> LocalIndex;
+					TArray<TArray<int32>>* LocalIndexPtr = &LocalIndex;
+					TArray<TArray<int>> GlobalIndex = Chaos::Utilities::ComputeIncidentElements(IndicesArray, LocalIndexPtr);
+					int32 ActualParticleCount = 0;
+					for (int32 l = 0; l < GlobalIndex.Num(); l++)
+					{
+						if (GlobalIndex[l].Num() > 0)
 						{
-							if (GlobalIndex[i].Num() > 0)
-							{
-								IndicesPositions[CurrentParticleIndex] = (*Vertices)[(*Indices)[GlobalIndex[i][0]][LocalIndex[i][0]]];
-								IndicesMap[CurrentParticleIndex] = (*Indices)[GlobalIndex[i][0]][LocalIndex[i][0]];
-								CurrentParticleIndex += 1;
-							}
+							ActualParticleCount += 1;
 						}
-						Chaos::FReal SphereRadius = (Chaos::FReal)0.;
+					}
+					TArray<Chaos::TVector<float, 3>> IndicesPositions; 
+					IndicesPositions.SetNum(ActualParticleCount);
+					TArray<int32> IndicesMap;
+					IndicesMap.SetNum(ActualParticleCount);
+					int32 CurrentParticleIndex = 0;
+					for (int32 i = 0; i < GlobalIndex.Num(); i++)
+					{
+						if (GlobalIndex[i].Num() > 0)
+						{
+							IndicesPositions[CurrentParticleIndex] = (*Vertices)[(*Indices)[GlobalIndex[i][0]][LocalIndex[i][0]]];
+							IndicesMap[CurrentParticleIndex] = (*Indices)[GlobalIndex[i][0]][LocalIndex[i][0]];
+							CurrentParticleIndex += 1;
+						}
+					}
+					Chaos::FReal SphereRadius = (Chaos::FReal)0.;
 
-						Chaos::TVec3<float> CoordMaxs(-FLT_MAX);
-						Chaos::TVec3<float> CoordMins(FLT_MAX);
-						for (int32 i = 0; i < IndicesPositions.Num(); i++)
+					Chaos::TVec3<float> CoordMaxs(-FLT_MAX);
+					Chaos::TVec3<float> CoordMins(FLT_MAX);
+					for (int32 i = 0; i < IndicesPositions.Num(); i++)
+					{
+						for (int32 j = 0; j < 3; j++)
 						{
-							for (int32 j = 0; j < 3; j++)
+							if (IndicesPositions[i][j] > CoordMaxs[j])
 							{
-								if (IndicesPositions[i][j] > CoordMaxs[j])
-								{
-									CoordMaxs[j] = IndicesPositions[i][j];
-								}
-								if (IndicesPositions[i][j] < CoordMins[j])
-								{
-									CoordMins[j] = IndicesPositions[i][j];
-								}
+								CoordMaxs[j] = IndicesPositions[i][j];
+							}
+							if (IndicesPositions[i][j] < CoordMins[j])
+							{
+								CoordMins[j] = IndicesPositions[i][j];
 							}
 						}
-						Chaos::TVec3<float> CoordDiff = (CoordMaxs - CoordMins) * VertexRadiusRatio;
-						SphereRadius = Chaos::FReal(FGenericPlatformMath::Min(CoordDiff[0], FGenericPlatformMath::Min(CoordDiff[1], CoordDiff[2])));
+					}
+					Chaos::TVec3<float> CoordDiff = (CoordMaxs - CoordMins) * VertexRadiusRatio;
+					SphereRadius = Chaos::FReal(FGenericPlatformMath::Min(CoordDiff[0], FGenericPlatformMath::Min(CoordDiff[1], CoordDiff[2])));
 
-						TArray<Chaos::TSphere<Chaos::FReal, 3>*> VertexSpherePtrs;
-						TArray<Chaos::TSphere<Chaos::FReal, 3>> VertexSpheres;
+					TArray<Chaos::TSphere<Chaos::FReal, 3>*> VertexSpherePtrs;
+					TArray<Chaos::TSphere<Chaos::FReal, 3>> VertexSpheres;
 
-						VertexSpheres.Init(Chaos::TSphere<Chaos::FReal, 3>(Chaos::TVec3<Chaos::FReal>(0), SphereRadius), IndicesPositions.Num());
-						VertexSpherePtrs.SetNum(IndicesPositions.Num());
+					VertexSpheres.Init(Chaos::TSphere<Chaos::FReal, 3>(Chaos::TVec3<Chaos::FReal>(0), SphereRadius), IndicesPositions.Num());
+					VertexSpherePtrs.SetNum(IndicesPositions.Num());
 
-						for (int32 i = 0; i < IndicesPositions.Num(); i++)
-						{
-							Chaos::TVec3<Chaos::FReal> SphereCenter(IndicesPositions[i]);
-							Chaos::TSphere<Chaos::FReal, 3> VertexSphere(SphereCenter, SphereRadius);
-							VertexSpheres[i] = Chaos::TSphere<Chaos::FReal, 3>(SphereCenter, SphereRadius);
-							VertexSpherePtrs[i] = &VertexSpheres[i];
-						}
-						Chaos::TBoundingVolumeHierarchy<
-							TArray<Chaos::TSphere<Chaos::FReal, 3>*>,
-							TArray<int32>,
-							Chaos::FReal,
-							3> VertexBVH(VertexSpherePtrs);
+					for (int32 i = 0; i < IndicesPositions.Num(); i++)
+					{
+						Chaos::TVec3<Chaos::FReal> SphereCenter(IndicesPositions[i]);
+						Chaos::TSphere<Chaos::FReal, 3> VertexSphere(SphereCenter, SphereRadius);
+						VertexSpheres[i] = Chaos::TSphere<Chaos::FReal, 3>(SphereCenter, SphereRadius);
+						VertexSpherePtrs[i] = &VertexSpheres[i];
+					}
+					Chaos::TBoundingVolumeHierarchy<
+						TArray<Chaos::TSphere<Chaos::FReal, 3>*>,
+						TArray<int32>,
+						Chaos::FReal,
+						3> VertexBVH(VertexSpherePtrs);
 
-						GeometryCollection::Facades::FPositionTargetFacade PositionTargets(InCollection);
-						PositionTargets.DefineSchema();
+					GeometryCollection::Facades::FPositionTargetFacade PositionTargets(InCollection);
+					PositionTargets.DefineSchema();
 
-						for (int32 i = 0; i < Indices->Num(); i++)
+					for (int32 i = 0; i < Indices->Num(); i++)
+					{
+						TArray<int32> TriangleIntersections0 = VertexBVH.FindAllIntersections((*Vertices)[(*Indices)[i][0]]);
+						TArray<int32> TriangleIntersections1 = VertexBVH.FindAllIntersections((*Vertices)[(*Indices)[i][1]]);
+						TArray<int32> TriangleIntersections2 = VertexBVH.FindAllIntersections((*Vertices)[(*Indices)[i][2]]);
+						TriangleIntersections0.Sort();
+						TriangleIntersections1.Sort();
+						TriangleIntersections2.Sort();
+
+						TArray<int32> TriangleIntersections({});
+						for (int32 k = 0; k < TriangleIntersections0.Num(); k++)
 						{
-							TArray<int32> TriangleIntersections0 = VertexBVH.FindAllIntersections((*Vertices)[(*Indices)[i][0]]);
-							TArray<int32> TriangleIntersections1 = VertexBVH.FindAllIntersections((*Vertices)[(*Indices)[i][1]]);
-							TArray<int32> TriangleIntersections2 = VertexBVH.FindAllIntersections((*Vertices)[(*Indices)[i][2]]);
-							TriangleIntersections0.Sort();
-							TriangleIntersections1.Sort();
-							TriangleIntersections2.Sort();
-
-							TArray<int32> TriangleIntersections({});
-							for (int32 k = 0; k < TriangleIntersections0.Num(); k++)
+							if (TriangleIntersections1.Contains(TriangleIntersections0[k]) 
+								&& TriangleIntersections2.Contains(TriangleIntersections0[k]))
 							{
-								if (TriangleIntersections1.Contains(TriangleIntersections0[k]) 
-									&& TriangleIntersections2.Contains(TriangleIntersections0[k]))
-								{
-									TriangleIntersections.Emplace(TriangleIntersections0[k]);
-								}
+								TriangleIntersections.Emplace(TriangleIntersections0[k]);
 							}
+						}
 
-							int32 TriangleIndex = (*ComponentIndex)[(*Indices)[i][0]];
-							int32 MinIndex = -1;
-							float MinDis = SphereRadius;
-							Chaos::TVector<float, 3> ClosestBary(0.f);
-							for (int32 j = 0; j < TriangleIntersections.Num(); j++)
+						int32 TriangleIndex = ComponentIndex[(*Indices)[i][0]];
+						int32 MinIndex = -1;
+						float MinDis = SphereRadius;
+						Chaos::TVector<float, 3> ClosestBary(0.f);
+						for (int32 j = 0; j < TriangleIntersections.Num(); j++)
+						{
+							if (ComponentIndex[IndicesMap[TriangleIntersections[j]]] >= 0 && TriangleIndex >= 0 && ComponentIndex[IndicesMap[TriangleIntersections[j]]] != TriangleIndex)
 							{
-								if ((*ComponentIndex)[IndicesMap[TriangleIntersections[j]]] >= 0 && TriangleIndex >= 0 && (*ComponentIndex)[IndicesMap[TriangleIntersections[j]]] != TriangleIndex)
+								Chaos::TVector<float, 3> Bary, TriPos0((*Vertices)[(*Indices)[i][0]]), TriPos1((*Vertices)[(*Indices)[i][1]]), TriPos2((*Vertices)[(*Indices)[i][2]]), ParticlePos((*Vertices)[IndicesMap[TriangleIntersections[j]]]);
+								Chaos::TVector<Chaos::FRealSingle, 3> ClosestPoint = Chaos::FindClosestPointAndBaryOnTriangle(TriPos0, TriPos1, TriPos2, ParticlePos, Bary);
+								Chaos::FRealSingle CurrentDistance = ((*Vertices)[IndicesMap[TriangleIntersections[j]]] - ClosestPoint).Size();
+								if (CurrentDistance < MinDis)
 								{
-									Chaos::TVector<float, 3> Bary, TriPos0((*Vertices)[(*Indices)[i][0]]), TriPos1((*Vertices)[(*Indices)[i][1]]), TriPos2((*Vertices)[(*Indices)[i][2]]), ParticlePos((*Vertices)[IndicesMap[TriangleIntersections[j]]]);
-									Chaos::TVector<Chaos::FRealSingle, 3> ClosestPoint = Chaos::FindClosestPointAndBaryOnTriangle(TriPos0, TriPos1, TriPos2, ParticlePos, Bary);
-									Chaos::FRealSingle CurrentDistance = ((*Vertices)[IndicesMap[TriangleIntersections[j]]] - ClosestPoint).Size();
-									if (CurrentDistance < MinDis)
-									{
-										MinDis = CurrentDistance;
-										MinIndex = IndicesMap[TriangleIntersections[j]];
-										ClosestBary = Bary;
-									}
-
+									MinDis = CurrentDistance;
+									MinIndex = IndicesMap[TriangleIntersections[j]];
+									ClosestBary = Bary;
 								}
+
 							}
-							if (MinIndex != -1
-								&& MinIndex != (*Indices)[i][0]
-								&& MinIndex != (*Indices)[i][1]
-								&& MinIndex != (*Indices)[i][2])
+						}
+						if (MinIndex != -1
+							&& MinIndex != (*Indices)[i][0]
+							&& MinIndex != (*Indices)[i][1]
+							&& MinIndex != (*Indices)[i][2])
+						{
+							GeometryCollection::Facades::FPositionTargetsData DataPackage;
+							DataPackage.TargetIndex.Init(MinIndex, 1);
+							DataPackage.TargetWeights.Init(1.f, 1);
+							DataPackage.SourceWeights.Init(1.f, 3);
+							DataPackage.SourceIndex.Init(-1, 3);
+							DataPackage.SourceIndex[0] = (*Indices)[i][0];
+							DataPackage.SourceIndex[1] = (*Indices)[i][1];
+							DataPackage.SourceIndex[2] = (*Indices)[i][2];
+							DataPackage.SourceWeights[0] = ClosestBary[0];
+							DataPackage.SourceWeights[1] = ClosestBary[1];
+							DataPackage.SourceWeights[2] = ClosestBary[2];
+							if (TManagedArray<float>* Mass = InCollection.FindAttribute<float>("Mass", FGeometryCollection::VerticesGroup))
 							{
-								GeometryCollection::Facades::FPositionTargetsData DataPackage;
-								DataPackage.TargetIndex.Init(MinIndex, 1);
-								DataPackage.TargetWeights.Init(1.f, 1);
-								DataPackage.SourceWeights.Init(1.f, 3);
-								DataPackage.SourceIndex.Init(-1, 3);
-								DataPackage.SourceIndex[0] = (*Indices)[i][0];
-								DataPackage.SourceIndex[1] = (*Indices)[i][1];
-								DataPackage.SourceIndex[2] = (*Indices)[i][2];
-								DataPackage.SourceWeights[0] = ClosestBary[0];
-								DataPackage.SourceWeights[1] = ClosestBary[1];
-								DataPackage.SourceWeights[2] = ClosestBary[2];
-								if (TManagedArray<float>* Mass = InCollection.FindAttribute<float>("Mass", FGeometryCollection::VerticesGroup))
+								DataPackage.Stiffness = 0.f;
+								for (int32 k = 0; k < 3; k++)
 								{
-									DataPackage.Stiffness = 0.f;
-									for (int32 k = 0; k < 3; k++)
-									{
-										DataPackage.Stiffness += DataPackage.SourceWeights[k] * PositionTargetStiffness * (*Mass)[DataPackage.SourceIndex[k]];
-									}
-									DataPackage.Stiffness += DataPackage.TargetWeights[0] * PositionTargetStiffness * (*Mass)[DataPackage.TargetIndex[0]];
-									DataPackage.Stiffness /= 2.f;
-								}
-								else
-								{
-									DataPackage.Stiffness = PositionTargetStiffness;
+									DataPackage.Stiffness += DataPackage.SourceWeights[k] * PositionTargetStiffness * (*Mass)[DataPackage.SourceIndex[k]];
 								}
-								PositionTargets.AddPositionTarget(DataPackage);
+								DataPackage.Stiffness += DataPackage.TargetWeights[0] * PositionTargetStiffness * (*Mass)[DataPackage.TargetIndex[0]];
+								DataPackage.Stiffness /= 2.f;
+							}
+							else
+							{
+								DataPackage.Stiffness = PositionTargetStiffness;
 							}
+							PositionTargets.AddPositionTarget(DataPackage);
 						}
 					}
 				}
diff --git a/Engine/Plugins/Experimental/ChaosFlesh/Source/ChaosFleshNodes/Public/Dataflow/ChaosFleshEngineAssetNodes.h b/Engine/Plugins/Experimental/ChaosFlesh/Source/ChaosFleshNodes/Public/Dataflow/ChaosFleshEngineAssetNodes.h
--- a/Engine/Plugins/Experimental/ChaosFlesh/Source/ChaosFleshNodes/Public/Dataflow/ChaosFleshEngineAssetNodes.h
+++ b/Engine/Plugins/Experimental/ChaosFlesh/Source/ChaosFleshNodes/Public/Dataflow/ChaosFleshEngineAssetNodes.h
@@ -310,7 +310,7 @@ public:
 	bool bDeleteHiddenFaces = false;
 
 	UPROPERTY(EditAnywhere, Category = "Dataflow")
-	FString TargetComponentIndex = "";
+	FString TargetGeometryIndex = "";
 
 	FIsolateComponentNode(const Dataflow::FNodeParameters& InParam, FGuid InGuid = FGuid::NewGuid())
 		: FDataflowNode(InParam, InGuid)
diff --git a/Engine/Plugins/Experimental/GeometryCollectionPlugin/Source/GeometryCollectionNodes/Private/Dataflow/GeometryCollectionTransferVertexScalarAttributeNode.cpp b/Engine/Plugins/Experimental/GeometryCollectionPlugin/Source/GeometryCollectionNodes/Private/Dataflow/GeometryCollectionTransferVertexScalarAttributeNode.cpp
--- a/Engine/Plugins/Experimental/GeometryCollectionPlugin/Source/GeometryCollectionNodes/Private/Dataflow/GeometryCollectionTransferVertexScalarAttributeNode.cpp
+++ b/Engine/Plugins/Experimental/GeometryCollectionPlugin/Source/GeometryCollectionNodes/Private/Dataflow/GeometryCollectionTransferVertexScalarAttributeNode.cpp
@@ -198,7 +198,7 @@ void FGeometryCollectionTransferVertexScalarAttributeNode::PairedGeometryTransfe
 					}
 				}
 				Chaos::TVec3<float> CoordDiff = (CoordMaxs - CoordMins) * VertexRadiusRatio;
-				SphereRadius = Chaos::FReal(FGenericPlatformMath::Min(CoordDiff[0], FGenericPlatformMath::Min(CoordDiff[1], CoordDiff[2])));
+				SphereRadius = Chaos::FReal(FGenericPlatformMath::Max(CoordDiff[0], FGenericPlatformMath::Max(CoordDiff[1], CoordDiff[2])));
 
 				TArray<Chaos::TSphere<Chaos::FReal, 3>*> VertexSpherePtrs;
 				TArray<Chaos::TSphere<Chaos::FReal, 3>> VertexSpheres;
@@ -349,7 +349,7 @@ void FGeometryCollectionTransferVertexScalarAttributeNode::NearestVertexTransfer
 		}
 	}
 	Chaos::TVec3<float> CoordDiff = (CoordMaxs - CoordMins) * VertexRadiusRatio;
-	SphereRadius = Chaos::FReal(FGenericPlatformMath::Min(CoordDiff[0], FGenericPlatformMath::Min(CoordDiff[1], CoordDiff[2])));
+	SphereRadius = Chaos::FReal(FGenericPlatformMath::Max(CoordDiff[0], FGenericPlatformMath::Max(CoordDiff[1], CoordDiff[2])));
 
 	TArray<Chaos::TSphere<Chaos::FReal, 3>*> VertexSpherePtrs;
 	TArray<Chaos::TSphere<Chaos::FReal, 3>> VertexSpheres;
diff --git a/Engine/Source/Runtime/Experimental/Chaos/Private/Chaos/Deformable/ChaosDeformableSolver.cpp b/Engine/Source/Runtime/Experimental/Chaos/Private/Chaos/Deformable/ChaosDeformableSolver.cpp
--- a/Engine/Source/Runtime/Experimental/Chaos/Private/Chaos/Deformable/ChaosDeformableSolver.cpp
+++ b/Engine/Source/Runtime/Experimental/Chaos/Private/Chaos/Deformable/ChaosDeformableSolver.cpp
@@ -31,6 +31,7 @@
 #include "GeometryCollection/Facades/CollectionVertexBoneWeightsFacade.h"
 #include "GeometryCollection/Facades/CollectionPositionTargetFacade.h"
 #include "GeometryCollection/Facades/CollectionConstraintOverrideFacade.h"
+#include "GeometryCollection/Facades/CollectionMeshFacade.h"
 #include "GeometryCollection/Facades/CollectionMuscleActivationFacade.h"
 #include "Misc/FileHelper.h"
 #include "HAL/PlatformFileManager.h"
@@ -1158,19 +1159,17 @@ namespace Chaos::Softs
 								ParticleComponentIndex[i + Offset] = ComponentOffset;
 							}
 							int NewComponentOffset = ComponentOffset;
-							if (const TManagedArray<int32>* ComponentIndex = Rest->FindAttribute<int32>("ComponentIndex", FGeometryCollection::VerticesGroup))
-							{
-								ensureMsgf(ComponentIndex->Num() == Vertex->Num(), TEXT("ComponentIndex size is not equal to vertex size"));
-								for (int32 i = 0; i < ComponentIndex->Num(); i++) {
-									if ((*ComponentIndex)[i] < 0)
-									{
-										ParticleComponentIndex[i + Offset] = (*ComponentIndex)[i]; //Isolated Nodes
-									}
-									else
-									{
-										ParticleComponentIndex[i + Offset] = ComponentOffset + (*ComponentIndex)[i];
-										NewComponentOffset = NewComponentOffset < ParticleComponentIndex[i + Offset] ? ParticleComponentIndex[i + Offset] : NewComponentOffset;
-									}
+							GeometryCollection::Facades::FCollectionMeshFacade MeshFacade(*Rest);
+							TArray<int32> ComponentIndex = MeshFacade.GetGeometryGroupIndexArray();
+							for (int32 i = 0; i < ComponentIndex.Num(); i++) {
+								if (ComponentIndex[i] < 0)
+								{
+									ParticleComponentIndex[i + Offset] = ComponentIndex[i]; //Isolated Nodes
+								}
+								else
+								{
+									ParticleComponentIndex[i + Offset] = ComponentOffset + ComponentIndex[i];
+									NewComponentOffset = NewComponentOffset < ParticleComponentIndex[i + Offset] ? ParticleComponentIndex[i + Offset] : NewComponentOffset;
 								}
 							}
 							ComponentOffset = NewComponentOffset + 1;
@@ -1201,7 +1200,7 @@ namespace Chaos::Softs
 					{
 						NRing = SurfaceTriangleMesh->GetNRing(VertexId, NRadius-1);
 					}
-					else
+					else if (NRadius == 1)
 					{
 						NRing.Add(VertexId);
 					}
@@ -1213,7 +1212,7 @@ namespace Chaos::Softs
 							TriangleSet.Add(CoincidentTriangles[j]);
 						}
 					}
-				}
+				},true //force single-threaded
 			);
 		}
 	}
diff --git a/Engine/Source/Runtime/Experimental/Chaos/Private/GeometryCollection/Facades/CollectionMeshFacade.cpp b/Engine/Source/Runtime/Experimental/Chaos/Private/GeometryCollection/Facades/CollectionMeshFacade.cpp
--- a/Engine/Source/Runtime/Experimental/Chaos/Private/GeometryCollection/Facades/CollectionMeshFacade.cpp
+++ b/Engine/Source/Runtime/Experimental/Chaos/Private/GeometryCollection/Facades/CollectionMeshFacade.cpp
@@ -166,6 +166,22 @@ namespace GeometryCollection::Facades
 		}
 	}
 
+	const TArray<int32> FCollectionMeshFacade::GetGeometryGroupIndexArray() const
+	{
+		TArray<int32> GroupIndexArray;
+		const TManagedArray<FVector3f>& Vertices = VertexAttribute.Get();
+		const TManagedArray<int32>& VertexStarts = VertexStartAttribute.Get();
+		const TManagedArray<int32>& VertexCounts = VertexCountAttribute.Get();
+		GroupIndexArray.Init(INDEX_NONE, Vertices.Num());
+		for (int32 GroupIdx = 0; GroupIdx < VertexStarts.Num(); ++GroupIdx)
+		{
+			for (int32 LocalIdx = 0; LocalIdx < VertexCounts[GroupIdx]; ++LocalIdx)
+			{
+				GroupIndexArray[VertexStarts[GroupIdx] + LocalIdx] = GroupIdx;
+			}
+		}
+		return GroupIndexArray;
+	}
 };
 
 
diff --git a/Engine/Source/Runtime/Experimental/Chaos/Public/GeometryCollection/Facades/CollectionMeshFacade.h b/Engine/Source/Runtime/Experimental/Chaos/Public/GeometryCollection/Facades/CollectionMeshFacade.h
--- a/Engine/Source/Runtime/Experimental/Chaos/Public/GeometryCollection/Facades/CollectionMeshFacade.h
+++ b/Engine/Source/Runtime/Experimental/Chaos/Public/GeometryCollection/Facades/CollectionMeshFacade.h
@@ -50,7 +50,9 @@ namespace GeometryCollection::Facades
 		/**
 		 * Bakes the transforms into the vertex positions and sets the bone transforms to identity
 		 */
-		 CHAOS_API void BakeTransform(int32 TransformIdx, const FTransform& InTransform);
+		CHAOS_API void BakeTransform(int32 TransformIdx, const FTransform& InTransform);
+
+		CHAOS_API const TArray<int32> GetGeometryGroupIndexArray() const;
 
 		TManagedArrayAccessor<int32> TransformToGeometryIndexAttribute;
 		TManagedArrayAccessor<int32> TransformIndexAttribute;
