ChaosFlesh : Node Restructuring

diff --git a/Engine/Plugins/Experimental/ChaosFlesh/Source/ChaosFleshNodes/Private/Dataflow/ChaosFleshAuthorSceneCollisionCandidatesNode.cpp b/Engine/Plugins/Experimental/ChaosFlesh/Source/ChaosFleshNodes/Private/Dataflow/ChaosFleshAuthorSceneCollisionCandidatesNode.cpp
new file mode 100644
--- /dev/null
+++ b/Engine/Plugins/Experimental/ChaosFlesh/Source/ChaosFleshNodes/Private/Dataflow/ChaosFleshAuthorSceneCollisionCandidatesNode.cpp
@@ -0,0 +1,89 @@
+// Copyright Epic Games, Inc. All Rights Reserved.
+
+#include "Dataflow/ChaosFleshAuthorSceneCollisionCandidatesNode.h"
+
+#include "ChaosFlesh/FleshCollection.h"
+#include "GeometryCollection/Facades/CollectionConstraintOverrideFacade.h"
+#include "Logging/LogMacros.h"
+
+#include UE_INLINE_GENERATED_CPP_BY_NAME(ChaosFleshAuthorSceneCollisionCandidatesNode)
+
+DEFINE_LOG_CATEGORY(LogAuthorSceneCollisionCandidates);
+
+void
+FAuthorSceneCollisionCandidates::Evaluate(Dataflow::FContext& Context, const FDataflowOutput* Out) const
+{
+	if (Out->IsA<DataType>(&Collection))
+	{
+		DataType InCollection = GetValue<DataType>(Context, &Collection);
+		const int32 BoneIndex = GetValue<int32>(Context, &OriginBoneIndex);
+
+		int32 Num = 0;
+		GeometryCollection::Facades::FConstraintOverrideCandidateFacade CnstrCandidates(InCollection);
+		CnstrCandidates.DefineSchema();
+
+		if (IsConnected(&VertexIndices))
+		{
+			if (TManagedArray<FVector3f>* Vertices = InCollection.FindAttribute<FVector3f>("Vertex", FGeometryCollection::VerticesGroup))
+			{
+				const TArray<int32>& Indices = GetValue<TArray<int32>>(Context, &VertexIndices);
+				for (int32 i = 0; i < Indices.Num(); i++)
+				{
+					if (Indices[i] >= 0 && Indices[i] < Vertices->Num())
+					{
+						GeometryCollection::Facades::FConstraintOverridesCandidateData Data;
+						Data.VertexIndex = Indices[i];
+						Data.BoneIndex = BoneIndex;
+						CnstrCandidates.Add(Data);
+						Num++;
+					}
+				}
+			}
+		}
+		else
+		{
+			TSet<int32> UniqueIndices;
+			if (bSurfaceVerticesOnly)
+			{
+				if (TManagedArray<FIntVector>* Indices =
+					InCollection.FindAttribute<FIntVector>("Indices", FGeometryCollection::FacesGroup))
+				{
+					for (int32 i = 0; i < Indices->Num(); i++)
+					{
+						for (int32 j = 0; j < 3; j++)
+						{
+							UniqueIndices.Add((*Indices)[i][j]);
+						}
+					}
+				}
+			}
+			else
+			{
+				if (TManagedArray<FIntVector4>* Indices =
+					InCollection.FindAttribute<FIntVector4>(
+						FTetrahedralCollection::TetrahedronAttribute, FTetrahedralCollection::TetrahedralGroup))
+				{
+					for (int32 i = 0; i < Indices->Num(); i++)
+					{
+						for (int32 j = 0; j < 4; j++)
+						{
+							UniqueIndices.Add((*Indices)[i][j]);
+						}
+					}
+				}
+			}
+			for (TSet<int32>::TConstIterator It = UniqueIndices.CreateConstIterator(); It; ++It)
+			{
+				GeometryCollection::Facades::FConstraintOverridesCandidateData Data;
+				Data.VertexIndex = *It;
+				Data.BoneIndex = BoneIndex;
+				CnstrCandidates.Add(Data);
+				Num++;
+			}
+		}
+		UE_LOG(LogAuthorSceneCollisionCandidates, Display,
+			TEXT("'%s' - Added %d scene collision candidates."),
+			*GetName().ToString(), Num);
+		SetValue(Context, MoveTemp(InCollection), &Collection);
+	}
+}
diff --git a/Engine/Plugins/Experimental/ChaosFlesh/Source/ChaosFleshNodes/Private/Dataflow/ChaosFleshCollisionBodyConstraintNode.cpp b/Engine/Plugins/Experimental/ChaosFlesh/Source/ChaosFleshNodes/Private/Dataflow/ChaosFleshCollisionBodyConstraintNode.cpp
new file mode 100644
--- /dev/null
+++ b/Engine/Plugins/Experimental/ChaosFlesh/Source/ChaosFleshNodes/Private/Dataflow/ChaosFleshCollisionBodyConstraintNode.cpp
@@ -0,0 +1,79 @@
+// Copyright Epic Games, Inc. All Rights Reserved.
+
+#include "Dataflow/ChaosFleshCollisionBodyConstraintNode.h"
+
+#include "Engine/SkeletalMesh.h"
+#include "GeometryCollection/Facades/CollectionKinematicBindingFacade.h"
+#include "GeometryCollection/Facades/CollectionVertexBoneWeightsFacade.h"
+#include "PhysicsEngine/PhysicsAsset.h"
+#include "PhysicsEngine/SkeletalBodySetup.h"
+#include "PhysicsEngine/SphylElem.h"
+
+#include UE_INLINE_GENERATED_CPP_BY_NAME(ChaosFleshCollisionBodyConstraintNode)
+
+void FKinematicBodySetupInitializationDataflowNode::Evaluate(Dataflow::FContext& Context, const FDataflowOutput* Out) const
+{
+	if (Out->IsA<DataType>(&Collection))
+	{
+		DataType InCollection = GetValue<DataType>(Context, &Collection);
+
+		if (TManagedArray<FVector3f>* Vertices = InCollection.FindAttribute<FVector3f>("Vertex", FGeometryCollection::VerticesGroup))
+		{	
+			if (TObjectPtr<USkeletalMesh> SkeletalMesh = GetValue<TObjectPtr<USkeletalMesh>>(Context, &SkeletalMeshIn))
+			{
+				if (UPhysicsAsset* PhysicsAsset = SkeletalMesh->GetPhysicsAsset())
+				{
+					TArray<TObjectPtr<USkeletalBodySetup>> SkeletalBodySetups = PhysicsAsset->SkeletalBodySetups;
+					TArray<FTransform> ComponentPose;
+					Dataflow::Animation::GlobalTransforms(SkeletalMesh->GetRefSkeleton(), ComponentPose);
+					for (const TObjectPtr<USkeletalBodySetup>& BodySetup : SkeletalBodySetups)
+					{	
+						TArray<FKSphylElem> SphylElems = BodySetup->AggGeom.SphylElems;
+						int32 BoneIndex = SkeletalMesh->GetRefSkeleton().FindBoneIndex(BodySetup->BoneName);
+						if (0 <= BoneIndex && BoneIndex < ComponentPose.Num())
+						{
+							TArray<int32> BoundVerts;
+							TArray<float> BoundWeights;
+							for (FKSphylElem Capsule : SphylElems)
+							{
+								for (int32 i = 0; i < Vertices->Num(); ++i)
+								{
+									float DistanceToCapsule = Capsule.GetShortestDistanceToPoint(FVector((*Vertices)[i]), Capsule.GetTransform());
+									DistanceToCapsule = Capsule.GetShortestDistanceToPoint(FVector((*Vertices)[i]), ComponentPose[BoneIndex]);
+									if (DistanceToCapsule < UE_SMALL_NUMBER)
+									{
+										if (BoundVerts.Find(i) == INDEX_NONE)
+										{
+											BoundVerts.Add(i);
+											BoundWeights.Add(1.0);
+										}
+									}
+								}
+							}
+							//get local coords of bound verts
+							typedef GeometryCollection::Facades::FKinematicBindingFacade FKinematics;
+							FKinematics Kinematics(InCollection); Kinematics.DefineSchema();
+							if (Kinematics.IsValid())
+							{
+								FKinematics::FBindingKey Binding = Kinematics.SetBoneBindings(BoneIndex, BoundVerts, BoundWeights);
+								TManagedArray<TArray<FVector3f>>& LocalPos = InCollection.AddAttribute<TArray<FVector3f>>("LocalPosition", Binding.GroupName);
+								Kinematics.AddKinematicBinding(Binding);
+
+								auto FloatVert = [](FVector3d V) { return FVector3f(V.X, V.Y, V.Z); };
+								auto DoubleVert = [](FVector3f V) { return FVector3d(V.X, V.Y, V.Z); };
+								LocalPos[Binding.Index].SetNum(BoundVerts.Num());
+								for (int32 i = 0; i < BoundVerts.Num(); i++)
+								{
+									FVector3f Temp = (*Vertices)[BoundVerts[i]];
+									LocalPos[Binding.Index][i] = FloatVert(ComponentPose[BoneIndex].InverseTransformPosition(DoubleVert(Temp)));
+								}
+							}
+						}
+					}
+				}
+				GeometryCollection::Facades::FVertexBoneWeightsFacade(InCollection).AddBoneWeightsFromKinematicBindings();
+			}
+		}
+		SetValue(Context, MoveTemp(InCollection), &Collection);
+	}
+}
diff --git a/Engine/Plugins/Experimental/ChaosFlesh/Source/ChaosFleshNodes/Private/Dataflow/ChaosFleshKinematicConstraintNode.cpp b/Engine/Plugins/Experimental/ChaosFlesh/Source/ChaosFleshNodes/Private/Dataflow/ChaosFleshKinematicConstraintNode.cpp
new file mode 100644
--- /dev/null
+++ b/Engine/Plugins/Experimental/ChaosFlesh/Source/ChaosFleshNodes/Private/Dataflow/ChaosFleshKinematicConstraintNode.cpp
@@ -0,0 +1,113 @@
+// Copyright Epic Games, Inc. All Rights Reserved.
+
+#include "Dataflow/ChaosFleshKinematicConstraintNode.h"
+
+#include "Engine/SkeletalMesh.h"
+#include "GeometryCollection/Facades/CollectionKinematicBindingFacade.h"
+#include "GeometryCollection/Facades/CollectionVertexBoneWeightsFacade.h"
+
+#include UE_INLINE_GENERATED_CPP_BY_NAME(ChaosFleshKinematicConstraintNode)
+
+
+void FKinematicInitializationDataflowNode::Evaluate(Dataflow::FContext& Context, const FDataflowOutput* Out) const
+{
+	if (Out->IsA<DataType>(&Collection))
+	{
+		DataType InCollection = GetValue<DataType>(Context, &Collection);
+
+		if (TManagedArray<FVector3f>* Vertices = InCollection.FindAttribute<FVector3f>("Vertex", FGeometryCollection::VerticesGroup))
+		{
+			if (FindInput(&VertexIndicesIn) && FindInput(&VertexIndicesIn)->GetConnection())
+			{
+				TArray<int32> BoundVerts;
+				TArray<float> BoundWeights;
+
+				for (int32 SelectionIndex : GetValue<TArray<int32>>(Context, &VertexIndicesIn))
+				{
+					if (0 <= SelectionIndex && SelectionIndex < Vertices->Num())
+					{
+						BoundVerts.Add(SelectionIndex);
+					}
+				}
+				if (BoundVerts.Num())
+				{
+					BoundWeights.Init(1.0, BoundVerts.Num());
+					GeometryCollection::Facades::FKinematicBindingFacade Kinematics(InCollection);
+					Kinematics.AddKinematicBinding(Kinematics.SetBoneBindings(INDEX_NONE, BoundVerts, BoundWeights));
+				}
+			}
+			else if (TObjectPtr<USkeletalMesh> SkeletalMesh = GetValue<TObjectPtr<USkeletalMesh>>(Context, &SkeletalMeshIn))
+			{
+				int32 IndexValue = GetValue<int32>(Context, &BoneIndexIn);
+				if (IndexValue != INDEX_NONE)
+				{
+					TArray<FTransform> ComponentPose;
+					Dataflow::Animation::GlobalTransforms(SkeletalMesh->GetRefSkeleton(), ComponentPose);
+
+					TArray<int32> BranchIndices;
+					if (SkeletalSelectionMode == ESkeletalSeletionMode::Dataflow_SkeletalSelection_Branch)
+					{
+						TArray<int32> ToProcess;
+						int32 CurrentIndex = IndexValue;
+						while (SkeletalMesh->GetRefSkeleton().IsValidIndex(CurrentIndex))
+						{
+							TArray<int32> Buffer;
+							SkeletalMesh->GetRefSkeleton().GetDirectChildBones(CurrentIndex, Buffer);
+
+							if (Buffer.Num())
+							{
+								ToProcess.Append(Buffer);
+							}
+
+							BranchIndices.Add(CurrentIndex);
+
+							CurrentIndex = INDEX_NONE;
+							if (ToProcess.Num())
+							{
+								CurrentIndex = ToProcess.Pop();
+							}
+						}
+					}
+					else // ESkeletalSeletionMode::Dataflow_SkeletalSelection_Single
+					{
+						BranchIndices.Add(IndexValue);
+					}
+
+					TSet<int32> ProcessedVertices;
+					for (const int32 Index : BranchIndices)
+					{
+						TArray<int32> BoundVerts;
+						TArray<float> BoundWeights;
+
+						if (0 <= Index && Index < ComponentPose.Num())
+						{
+							FVector3f BonePosition(ComponentPose[Index].GetTranslation());
+
+							int NumVertices = Vertices->Num();
+							for (int i = Vertices->Num() - 1; i > 0; i--)
+							{
+								if ((BonePosition - (*Vertices)[i]).Length() < Radius)
+								{
+									if (!ProcessedVertices.Contains(i))
+									{
+										ProcessedVertices.Add(i);
+										BoundVerts.Add(i);
+										BoundWeights.Add(1.0);
+									}
+								}
+							}
+
+							if (BoundVerts.Num())
+							{
+								GeometryCollection::Facades::FKinematicBindingFacade Kinematics(InCollection);
+								Kinematics.AddKinematicBinding(Kinematics.SetBoneBindings(Index, BoundVerts, BoundWeights));
+							}
+						}
+					}
+					GeometryCollection::Facades::FVertexBoneWeightsFacade(InCollection).AddBoneWeightsFromKinematicBindings();
+				}
+			}
+		}
+		SetValue(Context, MoveTemp(InCollection), &Collection);
+	}
+}
diff --git a/Engine/Plugins/Experimental/ChaosFlesh/Source/ChaosFleshNodes/Private/Dataflow/ChaosFleshKinematicInitializationNodes.cpp b/Engine/Plugins/Experimental/ChaosFlesh/Source/ChaosFleshNodes/Private/Dataflow/ChaosFleshKinematicInitializationNodes.cpp
deleted file mode 100644
diff --git a/Engine/Plugins/Experimental/ChaosFlesh/Source/ChaosFleshNodes/Private/Dataflow/ChaosFleshKinematicOriginInsertionInitializationNode.cpp b/Engine/Plugins/Experimental/ChaosFlesh/Source/ChaosFleshNodes/Private/Dataflow/ChaosFleshKinematicOriginInsertionInitializationNode.cpp
new file mode 100644
--- /dev/null
+++ b/Engine/Plugins/Experimental/ChaosFlesh/Source/ChaosFleshNodes/Private/Dataflow/ChaosFleshKinematicOriginInsertionInitializationNode.cpp
@@ -0,0 +1,137 @@
+// Copyright Epic Games, Inc. All Rights Reserved.
+
+#include "Dataflow/ChaosFleshKinematicOriginInsertionInitializationNode.h"
+
+#include "Engine/SkeletalMesh.h"
+#include "GeometryCollection/Facades/CollectionKinematicBindingFacade.h"
+#include "GeometryCollection/Facades/CollectionVertexBoneWeightsFacade.h"
+#include "Rendering/SkeletalMeshRenderData.h"
+#include "BoneWeights.h"
+
+#include UE_INLINE_GENERATED_CPP_BY_NAME(ChaosFleshKinematicOriginInsertionInitializationNode)
+
+
+void FKinematicOriginInsertionInitializationDataflowNode::Evaluate(Dataflow::FContext& Context, const FDataflowOutput* Out) const
+{
+	if (Out->IsA<DataType>(&Collection))
+	{
+		DataType InCollection = GetValue<DataType>(Context, &Collection);
+
+		if (TManagedArray<FVector3f>* Vertices = InCollection.FindAttribute<FVector3f>("Vertex", FGeometryCollection::VerticesGroup))
+		{
+			if (FindInput(&OriginVertexIndicesIn) && FindInput(&OriginVertexIndicesIn)->GetConnection() && FindInput(&InsertionVertexIndicesIn) && FindInput(&InsertionVertexIndicesIn)->GetConnection())
+			{
+				TArray<int32> BoundVerts;
+				TArray<float> BoundWeights;
+
+				for (int32 SelectionIndex : GetValue<TArray<int32>>(Context, &OriginVertexIndicesIn))
+				{
+					if (0 <= SelectionIndex && SelectionIndex < Vertices->Num())
+					{
+						BoundVerts.Add(SelectionIndex);
+					}
+				}
+				for (int32 SelectionIndex : GetValue<TArray<int32>>(Context, &InsertionVertexIndicesIn))
+				{
+					if (0 <= SelectionIndex && SelectionIndex < Vertices->Num())
+					{
+						BoundVerts.Add(SelectionIndex);
+					}
+				}
+				if (TObjectPtr<USkeletalMesh> BoneSkeletalMesh = GetValue<TObjectPtr<USkeletalMesh>>(Context, &BoneSkeletalMeshIn))
+				{
+					FSkeletalMeshRenderData* RenderData = BoneSkeletalMesh->GetResourceForRendering();
+					if (RenderData->LODRenderData.Num())
+					{
+						//Grab vertices only, no elements
+						FSkeletalMeshLODRenderData* LODRenderData = &RenderData->LODRenderData[0];
+						const FPositionVertexBuffer& PositionVertexBuffer =
+							LODRenderData->StaticVertexBuffers.PositionVertexBuffer;
+						//Grab skin weights
+						const FSkinWeightVertexBuffer* SkinWeightVertexBuffer = LODRenderData->GetSkinWeightVertexBuffer();
+						const int32 MaxBoneInfluences = SkinWeightVertexBuffer->GetMaxBoneInfluences();
+						TArray<FTransform> ComponentPose;
+						Dataflow::Animation::GlobalTransforms(BoneSkeletalMesh->GetRefSkeleton(), ComponentPose);
+						TArray<TArray<int32>> BoneBoundVerts;
+						TArray<TArray<float>> BoneBoundWeights;
+						BoneBoundVerts.SetNum(ComponentPose.Num());
+						BoneBoundWeights.SetNum(ComponentPose.Num());
+						if (!PositionVertexBuffer.GetNumVertices())
+						{
+							return;
+						}
+						auto DoubleVert = [](FVector3f V) { return FVector3d(V.X, V.Y, V.Z); };
+						for (int32 i = 0; i < BoundVerts.Num(); i++)
+						{
+							int ClosestPointIndex = 0;
+							double MinDistance = DBL_MAX;
+							for (uint32 j = 0; j < PositionVertexBuffer.GetNumVertices(); j++)
+							{
+								const FVector3f& Pos = PositionVertexBuffer.VertexPosition(j);
+								double Distance = FVector::Distance(DoubleVert((*Vertices)[BoundVerts[i]]), DoubleVert(Pos));
+								if (Distance < MinDistance)
+								{
+									ClosestPointIndex = j;
+									MinDistance = Distance;
+								}
+							}
+							int32 SectionIndex;
+							int32 VertIndex;
+							LODRenderData->GetSectionFromVertexIndex(ClosestPointIndex, SectionIndex, VertIndex);
+
+							check(SectionIndex < LODRenderData->RenderSections.Num());
+							const FSkelMeshRenderSection& Section = LODRenderData->RenderSections[SectionIndex];
+							int32 BufferVertIndex = Section.GetVertexBufferIndex() + VertIndex;
+							for (int32 InfluenceIndex = 0; InfluenceIndex < MaxBoneInfluences; InfluenceIndex++)
+							{
+								const int32 BoneIndex = Section.BoneMap[SkinWeightVertexBuffer->GetBoneIndex(BufferVertIndex, InfluenceIndex)];
+								const float	Weight = (float)SkinWeightVertexBuffer->GetBoneWeight(BufferVertIndex, InfluenceIndex) * UE::AnimationCore::InvMaxRawBoneWeightFloat;
+								if (Weight > float(0) && 0 <= BoneIndex && BoneIndex < ComponentPose.Num())
+								{
+									FString BoneName = BoneSkeletalMesh->GetRefSkeleton().GetBoneName(BoneIndex).ToString();
+									BoneBoundVerts[BoneIndex].Add(BoundVerts[i]);
+									BoneBoundWeights[BoneIndex].Add(Weight);
+								}
+							}
+						}
+						for (int32 BoneIndex = 0; BoneIndex < ComponentPose.Num(); ++BoneIndex)
+						{
+							if (BoneBoundVerts[BoneIndex].Num())
+							{
+								FString BoneName = BoneSkeletalMesh->GetRefSkeleton().GetBoneName(BoneIndex).ToString();
+								//get local coords of bound verts
+								typedef GeometryCollection::Facades::FKinematicBindingFacade FKinematics;
+								FKinematics Kinematics(InCollection); Kinematics.DefineSchema();
+								if (Kinematics.IsValid())
+								{
+									FKinematics::FBindingKey Binding = Kinematics.SetBoneBindings(BoneIndex, BoneBoundVerts[BoneIndex], BoneBoundWeights[BoneIndex]);
+									TManagedArray<TArray<FVector3f>>& LocalPos = InCollection.AddAttribute<TArray<FVector3f>>("LocalPosition", Binding.GroupName);
+									Kinematics.AddKinematicBinding(Binding);
+
+									auto FloatVert = [](FVector3d V) { return FVector3f(V.X, V.Y, V.Z); };
+									LocalPos[Binding.Index].SetNum(BoneBoundVerts[BoneIndex].Num());
+									for (int32 i = 0; i < BoneBoundVerts[BoneIndex].Num(); i++)
+									{
+										FVector3f Temp = (*Vertices)[BoneBoundVerts[BoneIndex][i]];
+										LocalPos[Binding.Index][i] = FloatVert(ComponentPose[BoneIndex].InverseTransformPosition(DoubleVert(Temp)));
+									}
+								}
+							}
+						}
+						GeometryCollection::Facades::FVertexBoneWeightsFacade(InCollection).AddBoneWeightsFromKinematicBindings();
+					}
+				}
+				else
+				{
+					if (BoundVerts.Num())
+					{
+						BoundWeights.Init(1.0, BoundVerts.Num());
+						GeometryCollection::Facades::FKinematicBindingFacade Kinematics(InCollection);
+						Kinematics.AddKinematicBinding(Kinematics.SetBoneBindings(INDEX_NONE, BoundVerts, BoundWeights));
+					}
+				}
+			}
+		}
+		SetValue(Context, MoveTemp(InCollection), &Collection);
+	}
+}
diff --git a/Engine/Plugins/Experimental/ChaosFlesh/Source/ChaosFleshNodes/Private/Dataflow/ChaosFleshKinematicTetrahedralConstraintNode.cpp b/Engine/Plugins/Experimental/ChaosFlesh/Source/ChaosFleshNodes/Private/Dataflow/ChaosFleshKinematicTetrahedralConstraintNode.cpp
new file mode 100644
--- /dev/null
+++ b/Engine/Plugins/Experimental/ChaosFlesh/Source/ChaosFleshNodes/Private/Dataflow/ChaosFleshKinematicTetrahedralConstraintNode.cpp
@@ -0,0 +1,122 @@
+// Copyright Epic Games, Inc. All Rights Reserved.
+
+#include "Dataflow/ChaosFleshKinematicTetrahedralConstraintNode.h"
+
+#include "ChaosFlesh/TetrahedralCollection.h"
+#include "Engine/SkeletalMesh.h"
+#include "GeometryCollection/Facades/CollectionKinematicBindingFacade.h"
+#include "GeometryCollection/Facades/CollectionVertexBoneWeightsFacade.h"
+
+#include UE_INLINE_GENERATED_CPP_BY_NAME(ChaosFleshKinematicTetrahedralConstraintNode)
+
+
+void FKinematicTetrahedralBindingsDataflowNode::Evaluate(Dataflow::FContext& Context, const FDataflowOutput* Out) const
+{
+	if (Out->IsA<DataType>(&Collection))
+	{
+		DataType InCollection = GetValue<DataType>(Context, &Collection);
+		TManagedArray<FIntVector4>* Tetrahedron = InCollection.FindAttribute<FIntVector4>(FTetrahedralCollection::TetrahedronAttribute, FTetrahedralCollection::TetrahedralGroup);
+		TManagedArray<FVector3f>* Vertex = InCollection.FindAttribute<FVector3f>("Vertex", "Vertices");
+
+		TObjectPtr<USkeletalMesh> SkeletalMesh = GetValue<TObjectPtr<USkeletalMesh>>(Context, &SkeletalMeshIn);
+		if (SkeletalMesh && Tetrahedron && Vertex)
+		{
+			//parse exclusion list to find bones to skip
+			TArray<FString> StrArray;
+			ExclusionList.ParseIntoArray(StrArray, *FString(" "));				
+			int32 NumTets = Tetrahedron->Num();
+			TArray<FTransform> ComponentPose;
+			Dataflow::Animation::GlobalTransforms(SkeletalMesh->GetRefSkeleton(), ComponentPose);
+			TArray<bool> VertAdded;
+			VertAdded.Init(false, Vertex->Num());
+			auto DoSkipBoneIndex = [&StrArray, &SkeletalMesh](int32 BoneIndex)
+			{
+				bool Skip = false;
+				FString BoneName = SkeletalMesh->GetRefSkeleton().GetBoneName(BoneIndex).ToString();
+				for (FString Elem : StrArray)
+				{
+					if (BoneName.Contains(Elem))
+					{
+						Skip = true;
+						break;
+					}
+				}
+				return Skip;
+			};
+			for (int32 b = 0; b < SkeletalMesh->GetRefSkeleton().GetNum(); ++b)
+			{
+				FVector3f BonePosition(ComponentPose[b].GetTranslation());
+				int32 ParentIndex=SkeletalMesh->GetRefSkeleton().GetParentIndex(b);
+				if (!(ParentIndex == INDEX_NONE || DoSkipBoneIndex(b) || DoSkipBoneIndex(ParentIndex)))
+				{
+					FVector3f ParentPosition(ComponentPose[ParentIndex].GetTranslation());
+					FVector3f RayDir = ParentPosition - BonePosition;
+					Chaos::FReal Length = RayDir.Length();
+					RayDir.Normalize();
+
+					if (Length > Chaos::FReal(1e-8)) 
+					{
+						TSet<int32> BoneVertSet;
+						for (int32 t = 0; t < NumTets; t++) 
+						{
+							int32 i = (*Tetrahedron)[t][0];
+							int32 j = (*Tetrahedron)[t][1];
+							int32 k = (*Tetrahedron)[t][2];
+							int32 l = (*Tetrahedron)[t][3];
+
+							TArray<Chaos::TVec3<Chaos::FRealSingle>> InVertices;
+							InVertices.SetNum(4);
+							InVertices[0][0] = (*Vertex)[i].X; InVertices[0][1] = (*Vertex)[i].Y; InVertices[0][2] = (*Vertex)[i].Z;
+							InVertices[1][0] = (*Vertex)[j].X; InVertices[1][1] = (*Vertex)[j].Y; InVertices[1][2] = (*Vertex)[j].Z;
+							InVertices[2][0] = (*Vertex)[k].X; InVertices[2][1] = (*Vertex)[k].Y; InVertices[2][2] = (*Vertex)[k].Z;
+							InVertices[3][0] = (*Vertex)[l].X; InVertices[3][1] = (*Vertex)[l].Y; InVertices[3][2] = (*Vertex)[l].Z;
+							Chaos::FConvex ConvexTet(InVertices, Chaos::FReal(0));
+							Chaos::FReal OutTime;
+							Chaos::FVec3 OutPosition, OutNormal;
+							int32 OutFaceIndex;
+							bool KeepTet = ConvexTet.Raycast(BonePosition, RayDir, Length, Chaos::FReal(0), OutTime, OutPosition, OutNormal, OutFaceIndex);
+							if (KeepTet) 
+							{	
+								for (int32 c = 0; c < 4; ++c)
+								{	
+									if (!VertAdded[(*Tetrahedron)[t][c]])
+									{
+										VertAdded[(*Tetrahedron)[t][c]] = true;
+										BoneVertSet.Add((*Tetrahedron)[t][c]);
+									}
+								}
+							}
+						}
+
+						TArray<int32> BoundVerts = BoneVertSet.Array();
+						TArray<float> BoundWeights;
+						BoundWeights.Init(float(1), BoundVerts.Num());
+						if (BoundVerts.Num())
+						{
+							//get local coords of bound verts
+							typedef GeometryCollection::Facades::FKinematicBindingFacade FKinematics;
+							FKinematics Kinematics(InCollection); Kinematics.DefineSchema();
+							if (Kinematics.IsValid())
+							{
+								FKinematics::FBindingKey Binding = Kinematics.SetBoneBindings(ParentIndex, BoundVerts, BoundWeights);
+								TManagedArray<TArray<FVector3f>>& LocalPos = InCollection.AddAttribute<TArray<FVector3f>>("LocalPosition", Binding.GroupName);
+								Kinematics.AddKinematicBinding(Binding);
+								auto DoubleVert = [](FVector3f V) { return FVector3d(V.X, V.Y, V.Z); };
+								auto FloatVert = [](FVector3d V) { return FVector3f(V.X, V.Y, V.Z); };
+								LocalPos[Binding.Index].SetNum(BoundVerts.Num());
+								for (int32 i = 0; i < BoundVerts.Num(); i++)
+								{
+									FVector3f Temp = (*Vertex)[BoundVerts[i]];
+									LocalPos[Binding.Index][i] = FloatVert(ComponentPose[ParentIndex].InverseTransformPosition(DoubleVert(Temp)));
+								}
+							}
+						}
+					}
+				}
+			}
+			GeometryCollection::Facades::FVertexBoneWeightsFacade(InCollection).AddBoneWeightsFromKinematicBindings();
+		}
+
+		SetValue(Context, MoveTemp(InCollection), &Collection);
+	}
+}
diff --git a/Engine/Plugins/Experimental/ChaosFlesh/Source/ChaosFleshNodes/Private/Dataflow/ChaosFleshNodesPlugin.cpp b/Engine/Plugins/Experimental/ChaosFlesh/Source/ChaosFleshNodes/Private/Dataflow/ChaosFleshNodesPlugin.cpp
--- a/Engine/Plugins/Experimental/ChaosFlesh/Source/ChaosFleshNodes/Private/Dataflow/ChaosFleshNodesPlugin.cpp
+++ b/Engine/Plugins/Experimental/ChaosFlesh/Source/ChaosFleshNodes/Private/Dataflow/ChaosFleshNodesPlugin.cpp
@@ -3,23 +3,31 @@
 #include "Dataflow/ChaosFleshNodesPlugin.h"
 
 
+#include "Dataflow/ChaosFleshAuthorSceneCollisionCandidatesNode.h"
 #include "Dataflow/ChaosFleshBindingsNodes.h"
+#include "Dataflow/ChaosFleshCollisionBodyConstraintNode.h"
 #include "Dataflow/ChaosFleshCoreNodes.h"
 #include "Dataflow/ChaosFleshCreateTetrahedronNode.h"
 #include "Dataflow/ChaosFleshEngineAssetNodes.h"
 #include "Dataflow/ChaosFleshFiberDirectionInitializationNodes.h"
 #include "Dataflow/ChaosFleshImportGEO.h"
-#include "Dataflow/ChaosFleshKinematicInitializationNodes.h"
+#include "Dataflow/ChaosFleshKinematicConstraintNode.h"
+#include "Dataflow/ChaosFleshKinematicOriginInsertionInitializationNode.h"
+#include "Dataflow/ChaosFleshKinematicTetrahedralConstraintNode.h"
 #include "Dataflow/ChaosFleshRadialTetrahedronNodes.h"
 #include "Dataflow/ChaosFleshRenderInitializationNodes.h"
 #include "Dataflow/ChaosFleshPositionTargetInitializationNodes.h"
 #include "Dataflow/ChaosFleshSkeletalBindingsNode.h"
 #include "Dataflow/ChaosFleshTetrahedralNodes.h"
 #include "Dataflow/ChaosFleshTriangleMeshSimulationPropertiesNode.h"
+#include "Dataflow/ChaosFleshSkeletalMeshConstraintNode.h"
 #include "Dataflow/ChaosFleshSkinSimulationPropertiesNode.h"
 #include "Dataflow/ChaosFleshSetFleshDefaultPropertiesNode.h"
+#include "Dataflow/ChaosFleshVertexConstraintNode.h"
+#include "Dataflow/GeometryCollectionAppendCollectionTransformNode.h"
 #include "Modules/ModuleManager.h"
 
+
 #define LOCTEXT_NAMESPACE "ChaosFleshNodes"
 
 
@@ -30,7 +38,6 @@ void IChaosFleshNodesPlugin::StartupModule()
 	Dataflow::RegisterChaosFleshCoreNodes();
 	Dataflow::ChaosFleshFiberDirectionInitializationNodes();
 	Dataflow::ChaosFleshRenderInitializationNodes();
-	Dataflow::RegisterChaosFleshKinematicInitializationNodes();
 	Dataflow::RegisterChaosFleshPositionTargetInitializationNodes();
 	Dataflow::ChaosFleshTetrahedralNodes();
 	Dataflow::ChaosFleshSkeletalBindingsNode();
@@ -40,6 +47,15 @@ void IChaosFleshNodesPlugin::StartupModule()
 	DATAFLOW_NODE_REGISTER_CREATION_FACTORY(FTriangleMeshSimulationPropertiesDataflowNodes);
 	DATAFLOW_NODE_REGISTER_CREATION_FACTORY(FSkinSimulationPropertiesDataflowNodes);
 	DATAFLOW_NODE_REGISTER_CREATION_FACTORY(FSetFleshDefaultPropertiesNode);
+
+	DATAFLOW_NODE_REGISTER_CREATION_FACTORY(FKinematicSkeletalMeshInitializationDataflowNode);
+	DATAFLOW_NODE_REGISTER_CREATION_FACTORY(FKinematicBodySetupInitializationDataflowNode);
+	DATAFLOW_NODE_REGISTER_CREATION_FACTORY(FKinematicInitializationDataflowNode);
+	DATAFLOW_NODE_REGISTER_CREATION_FACTORY(FKinematicOriginInsertionInitializationDataflowNode);
+	DATAFLOW_NODE_REGISTER_CREATION_FACTORY(FKinematicTetrahedralBindingsDataflowNode);
+	DATAFLOW_NODE_REGISTER_CREATION_FACTORY(FSetVerticesKinematicDataflowNode);
+	DATAFLOW_NODE_REGISTER_CREATION_FACTORY(FAuthorSceneCollisionCandidates);
+	DATAFLOW_NODE_REGISTER_CREATION_FACTORY(FAppendToCollectionTransformAttributeDataflowNode); 
 }
 
 void IChaosFleshNodesPlugin::ShutdownModule()
diff --git a/Engine/Plugins/Experimental/ChaosFlesh/Source/ChaosFleshNodes/Private/Dataflow/ChaosFleshSkeletalMeshConstraintNode.cpp b/Engine/Plugins/Experimental/ChaosFlesh/Source/ChaosFleshNodes/Private/Dataflow/ChaosFleshSkeletalMeshConstraintNode.cpp
new file mode 100644
--- /dev/null
+++ b/Engine/Plugins/Experimental/ChaosFlesh/Source/ChaosFleshNodes/Private/Dataflow/ChaosFleshSkeletalMeshConstraintNode.cpp
@@ -0,0 +1,99 @@
+// Copyright Epic Games, Inc. All Rights Reserved.
+
+#include "Dataflow/ChaosFleshSkeletalMeshConstraintNode.h"
+
+#include "Engine/SkeletalMesh.h"
+#include "GeometryCollection/Facades/CollectionKinematicBindingFacade.h"
+#include "GeometryCollection/Facades/CollectionVertexBoneWeightsFacade.h"
+#include "Rendering/SkeletalMeshRenderData.h"
+#include "BoneWeights.h"
+
+#include UE_INLINE_GENERATED_CPP_BY_NAME(ChaosFleshSkeletalMeshConstraintNode)
+
+
+void FKinematicSkeletalMeshInitializationDataflowNode::Evaluate(Dataflow::FContext& Context, const FDataflowOutput* Out) const
+{
+	if (Out->IsA<DataType>(&Collection) || Out->IsA<TArray<int32>>(&IndicesOut))
+	{
+		DataType InCollection = GetValue<DataType>(Context, &Collection);
+		TArray<int32> Indices;
+		if (TObjectPtr<USkeletalMesh> SkeletalMesh = GetValue<TObjectPtr<USkeletalMesh>>(Context, &SkeletalMeshIn))
+		{
+			FSkeletalMeshRenderData* RenderData = SkeletalMesh->GetResourceForRendering();
+			if (RenderData->LODRenderData.Num())
+			{
+				//Grab vertices only, no elements
+				FSkeletalMeshLODRenderData* LODRenderData = &RenderData->LODRenderData[0];
+				const FPositionVertexBuffer& PositionVertexBuffer =
+					LODRenderData->StaticVertexBuffers.PositionVertexBuffer;
+				TManagedArray<FVector3f>& Vertices = InCollection.AddAttribute<FVector3f>("Vertex", FGeometryCollection::VerticesGroup);
+				int32 index = InCollection.AddElements(PositionVertexBuffer.GetNumVertices(), FGeometryCollection::VerticesGroup);
+				for (uint32 j = 0; j < PositionVertexBuffer.GetNumVertices(); j++)
+				{
+					const FVector3f& Pos = PositionVertexBuffer.VertexPosition(j);
+					Vertices[index+j] = Pos;
+					Indices.Add(index+j);
+				}
+				//Grab skin weights
+				const FSkinWeightVertexBuffer* SkinWeightVertexBuffer = LODRenderData->GetSkinWeightVertexBuffer();
+				const int32 MaxBoneInfluences = SkinWeightVertexBuffer->GetMaxBoneInfluences();
+				TArray<FTransform> ComponentPose;
+				Dataflow::Animation::GlobalTransforms(SkeletalMesh->GetRefSkeleton(), ComponentPose);
+				TArray<TArray<int32>> BoundVerts;
+				TArray<TArray<float>> BoundWeights;
+				BoundVerts.SetNum(ComponentPose.Num());
+				BoundWeights.SetNum(ComponentPose.Num());
+				for (uint32 j = 0; j < PositionVertexBuffer.GetNumVertices(); j++)
+				{	
+					
+					int32 SectionIndex;
+					int32 VertIndex;
+					LODRenderData->GetSectionFromVertexIndex(j, SectionIndex, VertIndex);
+
+					check(SectionIndex < LODRenderData->RenderSections.Num());
+					const FSkelMeshRenderSection& Section = LODRenderData->RenderSections[SectionIndex];
+					int32 BufferVertIndex = Section.GetVertexBufferIndex() + VertIndex;
+					for (int32 InfluenceIndex = 0; InfluenceIndex < MaxBoneInfluences; InfluenceIndex++)
+					{
+						const int32 BoneIndex = Section.BoneMap[SkinWeightVertexBuffer->GetBoneIndex(BufferVertIndex, InfluenceIndex)];
+						const float	Weight = (float)SkinWeightVertexBuffer->GetBoneWeight(BufferVertIndex, InfluenceIndex) * UE::AnimationCore::InvMaxRawBoneWeightFloat;
+						if (Weight > float(0) && 0 <= BoneIndex && BoneIndex < ComponentPose.Num())
+						{
+							FString BoneName = SkeletalMesh->GetRefSkeleton().GetBoneName(BoneIndex).ToString();
+							BoundVerts[BoneIndex].Add(index+j);
+							BoundWeights[BoneIndex].Add(Weight);
+						}
+					}
+				}
+				for (int32 BoneIndex = 0; BoneIndex < ComponentPose.Num(); ++BoneIndex)
+				{
+					if (BoundVerts[BoneIndex].Num())
+					{
+						FString BoneName = SkeletalMesh->GetRefSkeleton().GetBoneName(BoneIndex).ToString();
+						//get local coords of bound verts
+						typedef GeometryCollection::Facades::FKinematicBindingFacade FKinematics;
+						FKinematics Kinematics(InCollection); Kinematics.DefineSchema();
+						if (Kinematics.IsValid())
+						{
+							FKinematics::FBindingKey Binding = Kinematics.SetBoneBindings(BoneIndex, BoundVerts[BoneIndex], BoundWeights[BoneIndex]);
+							TManagedArray<TArray<FVector3f>>& LocalPos = InCollection.AddAttribute<TArray<FVector3f>>("LocalPosition", Binding.GroupName);
+							Kinematics.AddKinematicBinding(Binding);
+
+							auto FloatVert = [](FVector3d V) { return FVector3f(V.X, V.Y, V.Z); };
+							auto DoubleVert = [](FVector3f V) { return FVector3d(V.X, V.Y, V.Z); };
+							LocalPos[Binding.Index].SetNum(BoundVerts[BoneIndex].Num());
+							for (int32 i = 0; i < BoundVerts[BoneIndex].Num(); i++)
+							{
+								FVector3f Temp = Vertices[BoundVerts[BoneIndex][i]];
+								LocalPos[Binding.Index][i] = FloatVert(ComponentPose[BoneIndex].InverseTransformPosition(DoubleVert(Temp)));
+							}
+						}
+					}
+				}
+				GeometryCollection::Facades::FVertexBoneWeightsFacade(InCollection).AddBoneWeightsFromKinematicBindings();
+			}
+		}
+		SetValue(Context, MoveTemp(InCollection), &Collection);
+		SetValue(Context, MoveTemp(Indices), &IndicesOut);
+	}
+}
diff --git a/Engine/Plugins/Experimental/ChaosFlesh/Source/ChaosFleshNodes/Private/Dataflow/ChaosFleshVertexConstraintNode.cpp b/Engine/Plugins/Experimental/ChaosFlesh/Source/ChaosFleshNodes/Private/Dataflow/ChaosFleshVertexConstraintNode.cpp
new file mode 100644
--- /dev/null
+++ b/Engine/Plugins/Experimental/ChaosFlesh/Source/ChaosFleshNodes/Private/Dataflow/ChaosFleshVertexConstraintNode.cpp
@@ -0,0 +1,39 @@
+// Copyright Epic Games, Inc. All Rights Reserved.
+
+#include "Dataflow/ChaosFleshVertexConstraintNode.h"
+
+#include "GeometryCollection/Facades/CollectionKinematicBindingFacade.h"
+
+#include UE_INLINE_GENERATED_CPP_BY_NAME(ChaosFleshVertexConstraintNode)
+
+void FSetVerticesKinematicDataflowNode::Evaluate(Dataflow::FContext& Context, const FDataflowOutput* Out) const
+{
+	if (Out->IsA<DataType>(&Collection))
+	{
+		DataType InCollection = GetValue<DataType>(Context, &Collection);
+		TArray<int32> BoundVerts;
+		TArray<float> BoundWeights;
+		if (FindInput(&VertexIndicesIn) && FindInput(&VertexIndicesIn)->GetConnection())
+		{
+			if (TManagedArray<FVector3f>* Vertices = InCollection.FindAttribute<FVector3f>("Vertex", FGeometryCollection::VerticesGroup))
+			{
+
+				for (int32 SelectionIndex : GetValue<TArray<int32>>(Context, &VertexIndicesIn))
+				{
+					if (0 <= SelectionIndex && SelectionIndex < Vertices->Num())
+					{
+						BoundVerts.Add(SelectionIndex);
+					}
+				}
+				BoundWeights.Init(1.0, BoundVerts.Num());	
+			}
+		} 
+
+		if (BoundVerts.Num() > 0)
+		{
+			GeometryCollection::Facades::FKinematicBindingFacade Kinematics(InCollection);
+			Kinematics.AddKinematicBinding(Kinematics.SetBoneBindings(INDEX_NONE, BoundVerts, BoundWeights));
+		}
+		SetValue(Context, MoveTemp(InCollection), &Collection);
+	}
+}
diff --git a/Engine/Plugins/Experimental/ChaosFlesh/Source/ChaosFleshNodes/Private/Dataflow/GeometryCollectionAppendCollectionTransformNode.cpp b/Engine/Plugins/Experimental/ChaosFlesh/Source/ChaosFleshNodes/Private/Dataflow/GeometryCollectionAppendCollectionTransformNode.cpp
new file mode 100644
--- /dev/null
+++ b/Engine/Plugins/Experimental/ChaosFlesh/Source/ChaosFleshNodes/Private/Dataflow/GeometryCollectionAppendCollectionTransformNode.cpp
@@ -0,0 +1,29 @@
+// Copyright Epic Games, Inc. All Rights Reserved.
+
+#include "Dataflow/GeometryCollectionAppendCollectionTransformNode.h"
+
+#include UE_INLINE_GENERATED_CPP_BY_NAME(GeometryCollectionAppendCollectionTransformNode)
+
+void
+FAppendToCollectionTransformAttributeDataflowNode::Evaluate(Dataflow::FContext& Context, const FDataflowOutput* Out) const
+{
+	if (Out->IsA<DataType>(&Collection))
+	{
+		const FTransform& Transform = GetValue<FTransform>(Context, &TransformIn);
+		FManagedArrayCollection CollectionValue = GetValue<FManagedArrayCollection>(Context, &Collection);
+		const TManagedArray<FTransform>* ConstTransformsPtr = CollectionValue.FindAttributeTyped<FTransform>(FName(AttributeName), FName(GroupName));
+		if (!ConstTransformsPtr)
+		{
+			CollectionValue.AddAttribute<FTransform>(FName(AttributeName), FName(GroupName));
+			ConstTransformsPtr = CollectionValue.FindAttributeTyped<FTransform>(FName(AttributeName), FName(GroupName));
+		}
+
+		if(ConstTransformsPtr)
+		{
+			int32 Index = CollectionValue.AddElements(1, FName(GroupName));
+			TManagedArray<FTransform>& Transforms = CollectionValue.ModifyAttribute<FTransform>(FName(AttributeName), FName(GroupName));
+			Transforms[Index] = Transform;
+		}
+		SetValue(Context, MoveTemp(CollectionValue), &Collection);
+	}
+}
\ No newline at end of file
diff --git a/Engine/Plugins/Experimental/ChaosFlesh/Source/ChaosFleshNodes/Public/Dataflow/ChaosFleshAuthorSceneCollisionCandidatesNode.h b/Engine/Plugins/Experimental/ChaosFlesh/Source/ChaosFleshNodes/Public/Dataflow/ChaosFleshAuthorSceneCollisionCandidatesNode.h
new file mode 100644
--- /dev/null
+++ b/Engine/Plugins/Experimental/ChaosFlesh/Source/ChaosFleshNodes/Public/Dataflow/ChaosFleshAuthorSceneCollisionCandidatesNode.h
@@ -0,0 +1,52 @@
+// Copyright Epic Games, Inc. All Rights Reserved.
+#pragma once 
+
+#include "CoreMinimal.h"
+#include "Dataflow/DataflowCore.h"
+#include "Dataflow/DataflowEngine.h"
+
+#include "ChaosFleshAuthorSceneCollisionCandidatesNode.generated.h"
+
+DECLARE_LOG_CATEGORY_EXTERN(LogAuthorSceneCollisionCandidates, Verbose, All);
+
+/**
+* Marks mesh vertices as candidates for scene collision raycasts.  Each vertex has an associated 
+* bone index to use as the origin of the raycast.  The runtime distance between the vertex and the
+* bone designates the range of the scene query depth.
+*/
+USTRUCT(meta = (DataflowFlesh))
+struct FAuthorSceneCollisionCandidates : public FDataflowNode
+{
+	GENERATED_USTRUCT_BODY()
+	DATAFLOW_NODE_DEFINE_INTERNAL(FAuthorSceneCollisionCandidates, "AuthorSceneCollisionCandidates", "Flesh", "")
+
+public:
+	typedef FManagedArrayCollection DataType;
+
+	UPROPERTY(meta = (DataflowInput, DataflowOutput, DisplayName = "Collection", DataflowPassthrough = "Collection"))
+	FManagedArrayCollection Collection;
+
+	//! Restricts vertices to only ones on the surface.  All vertices otherwise.
+	UPROPERTY(EditAnywhere, Category = "Dataflow")
+	bool bSurfaceVerticesOnly = true;
+
+	//! Indices to use with environment collisions.  If this input is not connected, then all 
+	//! indicies are used.
+	UPROPERTY(EditAnywhere, Category = "Dataflow", meta = (DataflowInput, DisplayName = "Vertex Indices"))
+	TArray<int32> VertexIndices;
+
+	//! Bone index to use as the world raycast origin.  -1 denotes the component transform.
+	UPROPERTY(EditAnywhere, Category = "Dataflow", meta = (DataflowInput, DisplayName = "Raycast Origin Bone Index"))
+	int32 OriginBoneIndex = 0;
+
+	FAuthorSceneCollisionCandidates(const Dataflow::FNodeParameters& InParam, FGuid InGuid = FGuid::NewGuid())
+		: FDataflowNode(InParam, InGuid)
+	{
+		RegisterInputConnection(&Collection);
+		RegisterOutputConnection(&Collection, &Collection);
+		RegisterInputConnection(&VertexIndices);
+		RegisterInputConnection(&OriginBoneIndex);
+	}
+
+	virtual void Evaluate(Dataflow::FContext& Context, const FDataflowOutput* Out) const override;
+};
diff --git a/Engine/Plugins/Experimental/ChaosFlesh/Source/ChaosFleshNodes/Public/Dataflow/ChaosFleshCollisionBodyConstraintNode.h b/Engine/Plugins/Experimental/ChaosFlesh/Source/ChaosFleshNodes/Public/Dataflow/ChaosFleshCollisionBodyConstraintNode.h
new file mode 100644
--- /dev/null
+++ b/Engine/Plugins/Experimental/ChaosFlesh/Source/ChaosFleshNodes/Public/Dataflow/ChaosFleshCollisionBodyConstraintNode.h
@@ -0,0 +1,43 @@
+// Copyright Epic Games, Inc. All Rights Reserved.
+#pragma once 
+
+#include "CoreMinimal.h"
+#include "Dataflow/DataflowCore.h"
+#include "Dataflow/DataflowEngine.h"
+
+#include "ChaosFleshCollisionBodyConstraintNode.generated.h"
+
+class USkeletalMesh;
+
+// @todo(deprecate), rename to FCollisionBodyConstraintDataflowNode
+
+USTRUCT(meta = (DataflowFlesh))
+struct FKinematicBodySetupInitializationDataflowNode : public FDataflowNode
+{
+	GENERATED_USTRUCT_BODY()
+	DATAFLOW_NODE_DEFINE_INTERNAL(FKinematicBodySetupInitializationDataflowNode, "KinematicBodySetupInitialization", "Flesh", "")
+	DATAFLOW_NODE_RENDER_TYPE(FGeometryCollection::StaticType(), "Collection")
+
+public:
+	typedef FManagedArrayCollection DataType;
+
+	UPROPERTY(EditAnywhere, Category = "Dataflow")
+		FTransform Transform;
+
+	UPROPERTY(meta = (DataflowInput, DataflowOutput, DisplayName = "Collection", DataflowPassthrough = "Collection"))
+		FManagedArrayCollection Collection;
+
+	UPROPERTY(meta = (DataflowInput, DisplayName = "SkeletalMesh"))
+		TObjectPtr<USkeletalMesh> SkeletalMeshIn;
+
+	FKinematicBodySetupInitializationDataflowNode(const Dataflow::FNodeParameters& InParam, FGuid InGuid = FGuid::NewGuid())
+		: FDataflowNode(InParam, InGuid)
+	{
+		RegisterInputConnection(&Collection);
+		RegisterOutputConnection(&Collection, &Collection);
+		RegisterInputConnection(&SkeletalMeshIn);
+	}
+
+	virtual void Evaluate(Dataflow::FContext& Context, const FDataflowOutput* Out) const override;
+};
+
diff --git a/Engine/Plugins/Experimental/ChaosFlesh/Source/ChaosFleshNodes/Public/Dataflow/ChaosFleshKinematicConstraintNode.h b/Engine/Plugins/Experimental/ChaosFlesh/Source/ChaosFleshNodes/Public/Dataflow/ChaosFleshKinematicConstraintNode.h
new file mode 100644
--- /dev/null
+++ b/Engine/Plugins/Experimental/ChaosFlesh/Source/ChaosFleshNodes/Public/Dataflow/ChaosFleshKinematicConstraintNode.h
@@ -0,0 +1,54 @@
+// Copyright Epic Games, Inc. All Rights Reserved.
+#pragma once 
+
+#include "CoreMinimal.h"
+#include "Dataflow/DataflowCore.h"
+#include "Dataflow/DataflowEngine.h"
+#include "Dataflow/ChaosFleshNodesUtility.h"
+
+#include "ChaosFleshKinematicConstraintNode.generated.h"
+
+class USkeletalMesh;
+
+// @todo(deprecate), rename to FKinematicConstraintDataflowNode
+USTRUCT(meta = (DataflowFlesh))
+struct FKinematicInitializationDataflowNode : public FDataflowNode
+{
+	GENERATED_USTRUCT_BODY()
+	DATAFLOW_NODE_DEFINE_INTERNAL(FKinematicInitializationDataflowNode, "KinematicInitialization", "Flesh", "")
+	DATAFLOW_NODE_RENDER_TYPE(FGeometryCollection::StaticType(), "Collection")
+
+public:
+	typedef FManagedArrayCollection DataType;
+
+	UPROPERTY(EditAnywhere, Category = "Dataflow")
+	float Radius = 40.f;
+
+	UPROPERTY(EditAnywhere, Category = "Dataflow")
+	ESkeletalSeletionMode SkeletalSelectionMode = ESkeletalSeletionMode::Dataflow_SkeletalSelection_Single;
+
+	UPROPERTY(meta = (DataflowInput, DataflowOutput, DisplayName = "Collection", DataflowPassthrough  = "Collection"))
+	FManagedArrayCollection Collection;
+
+	UPROPERTY(meta = (DataflowInput, DisplayName = "SkeletalMesh"))
+	TObjectPtr<USkeletalMesh> SkeletalMeshIn;
+		
+	UPROPERTY(meta = (DataflowInput, DisplayName = "SelectionSet"))
+	TArray<int32> VertexIndicesIn;
+
+	UPROPERTY(meta = (DataflowInput, DisplayName = "BoneIndex"))
+	int32 BoneIndexIn = INDEX_NONE;
+	
+
+	FKinematicInitializationDataflowNode(const Dataflow::FNodeParameters& InParam, FGuid InGuid = FGuid::NewGuid())
+		: FDataflowNode(InParam, InGuid)
+	{
+		RegisterInputConnection(&Collection);
+		RegisterOutputConnection(&Collection, &Collection);
+		RegisterInputConnection(&SkeletalMeshIn);
+		RegisterInputConnection(&VertexIndicesIn);
+		RegisterInputConnection(&BoneIndexIn);
+	}
+
+	virtual void Evaluate(Dataflow::FContext& Context, const FDataflowOutput* Out) const override;
+};
diff --git a/Engine/Plugins/Experimental/ChaosFlesh/Source/ChaosFleshNodes/Public/Dataflow/ChaosFleshKinematicInitializationNodes.h b/Engine/Plugins/Experimental/ChaosFlesh/Source/ChaosFleshNodes/Public/Dataflow/ChaosFleshKinematicInitializationNodes.h
deleted file mode 100644
diff --git a/Engine/Plugins/Experimental/ChaosFlesh/Source/ChaosFleshNodes/Public/Dataflow/ChaosFleshKinematicOriginInsertionInitializationNode.h b/Engine/Plugins/Experimental/ChaosFlesh/Source/ChaosFleshNodes/Public/Dataflow/ChaosFleshKinematicOriginInsertionInitializationNode.h
new file mode 100644
--- /dev/null
+++ b/Engine/Plugins/Experimental/ChaosFlesh/Source/ChaosFleshNodes/Public/Dataflow/ChaosFleshKinematicOriginInsertionInitializationNode.h
@@ -0,0 +1,48 @@
+// Copyright Epic Games, Inc. All Rights Reserved.
+#pragma once 
+
+#include "CoreMinimal.h"
+#include "Dataflow/DataflowCore.h"
+#include "Dataflow/DataflowEngine.h"
+
+#include "ChaosFleshKinematicOriginInsertionInitializationNode.generated.h"
+
+class USkeletalMesh;
+
+// @todo(deprecate), rename to FKinematicSomethingBetterThanThisNameDataflowNode
+USTRUCT(meta = (DataflowFlesh))
+struct FKinematicOriginInsertionInitializationDataflowNode : public FDataflowNode
+{
+	GENERATED_USTRUCT_BODY()
+	DATAFLOW_NODE_DEFINE_INTERNAL(FKinematicOriginInsertionInitializationDataflowNode, "KinematicOriginInsertionInitialization", "Flesh", "")
+	DATAFLOW_NODE_RENDER_TYPE(FGeometryCollection::StaticType(), "Collection")
+
+public:
+	typedef FManagedArrayCollection DataType;
+
+	UPROPERTY(meta = (DataflowInput, DataflowOutput, DisplayName = "Collection", DataflowPassthrough = "Collection"))
+		FManagedArrayCollection Collection;
+
+	UPROPERTY(meta = (DataflowInput, DisplayName = "OriginSelectionSet"))
+		TArray<int32> OriginVertexIndicesIn;
+
+	UPROPERTY(meta = (DataflowInput, DisplayName = "InsertionSelectionSet"))
+		TArray<int32> InsertionVertexIndicesIn;
+
+	UPROPERTY(meta = (DataflowInput, DisplayName = "BoneSkeletalMesh"))
+		TObjectPtr<USkeletalMesh> BoneSkeletalMeshIn;
+
+
+	FKinematicOriginInsertionInitializationDataflowNode(const Dataflow::FNodeParameters& InParam, FGuid InGuid = FGuid::NewGuid())
+		: FDataflowNode(InParam, InGuid)
+	{
+		RegisterInputConnection(&Collection);
+		RegisterOutputConnection(&Collection, &Collection);
+		RegisterInputConnection(&OriginVertexIndicesIn);
+		RegisterInputConnection(&InsertionVertexIndicesIn);
+		RegisterInputConnection(&BoneSkeletalMeshIn);
+	}
+
+	virtual void Evaluate(Dataflow::FContext& Context, const FDataflowOutput* Out) const override;
+};
+
diff --git a/Engine/Plugins/Experimental/ChaosFlesh/Source/ChaosFleshNodes/Public/Dataflow/ChaosFleshKinematicTetrahedralConstraintNode.h b/Engine/Plugins/Experimental/ChaosFlesh/Source/ChaosFleshNodes/Public/Dataflow/ChaosFleshKinematicTetrahedralConstraintNode.h
new file mode 100644
--- /dev/null
+++ b/Engine/Plugins/Experimental/ChaosFlesh/Source/ChaosFleshNodes/Public/Dataflow/ChaosFleshKinematicTetrahedralConstraintNode.h
@@ -0,0 +1,43 @@
+// Copyright Epic Games, Inc. All Rights Reserved.
+#pragma once 
+
+#include "CoreMinimal.h"
+#include "Dataflow/DataflowCore.h"
+#include "Dataflow/DataflowEngine.h"
+
+#include "ChaosFleshKinematicTetrahedralConstraintNode.generated.h"
+
+class USkeletalMesh;
+
+// @todo(deprecate), rename to FKinematicTetrahedralConstraintDataflowNode
+// ... This should really be a mode on the KinematicConstraint. 
+USTRUCT(meta = (DataflowFlesh))
+struct FKinematicTetrahedralBindingsDataflowNode : public FDataflowNode
+{
+	GENERATED_USTRUCT_BODY()
+	DATAFLOW_NODE_DEFINE_INTERNAL(FKinematicTetrahedralBindingsDataflowNode, "KinematicTetrahedralBindings", "Flesh", "")
+	DATAFLOW_NODE_RENDER_TYPE(FGeometryCollection::StaticType(), "Collection")
+
+public:
+	typedef FManagedArrayCollection DataType;
+
+	UPROPERTY(meta = (DataflowInput, DisplayName = "SkeletalMesh"))
+	TObjectPtr<USkeletalMesh> SkeletalMeshIn = nullptr;
+
+	UPROPERTY(meta = (DataflowInput, DataflowOutput, DisplayName = "Collection", DataflowPassthrough  = "Collection"))
+	FManagedArrayCollection Collection;
+
+	UPROPERTY(EditAnywhere, Category = "Dataflow")
+	FString ExclusionList = "twist foo";
+
+	FKinematicTetrahedralBindingsDataflowNode(const Dataflow::FNodeParameters& InParam, FGuid InGuid = FGuid::NewGuid())
+		: FDataflowNode(InParam, InGuid)
+	{
+		RegisterInputConnection(&SkeletalMeshIn);
+		RegisterInputConnection(&Collection);
+		RegisterOutputConnection(&Collection, &Collection);
+	}
+
+	virtual void Evaluate(Dataflow::FContext& Context, const FDataflowOutput* Out) const override;
+};
+
diff --git a/Engine/Plugins/Experimental/ChaosFlesh/Source/ChaosFleshNodes/Public/Dataflow/ChaosFleshNodesUtility.h b/Engine/Plugins/Experimental/ChaosFlesh/Source/ChaosFleshNodes/Public/Dataflow/ChaosFleshNodesUtility.h
--- a/Engine/Plugins/Experimental/ChaosFlesh/Source/ChaosFleshNodes/Public/Dataflow/ChaosFleshNodesUtility.h
+++ b/Engine/Plugins/Experimental/ChaosFlesh/Source/ChaosFleshNodes/Public/Dataflow/ChaosFleshNodesUtility.h
@@ -15,6 +15,15 @@ enum TetMeshingMethod : int
 	TetWild			UMETA(DisplayName = "TetWild"),
 };
 
+UENUM(BlueprintType)
+enum class ESkeletalSeletionMode : uint8
+{
+	Dataflow_SkeletalSelection_Single UMETA(DisplayName = "Single"),
+	Dataflow_SkeletalSelection_Branch UMETA(DisplayName = "Sub-Branch"),
+	//
+	Chaos_Max UMETA(Hidden)
+};
+
 namespace Dataflow
 {
 	TArray<FIntVector3> GetSurfaceTriangles(const TArray<FIntVector4>& Tets, const bool bKeepInterior);
diff --git a/Engine/Plugins/Experimental/ChaosFlesh/Source/ChaosFleshNodes/Public/Dataflow/ChaosFleshPositionTargetInitializationNodes.h b/Engine/Plugins/Experimental/ChaosFlesh/Source/ChaosFleshNodes/Public/Dataflow/ChaosFleshPositionTargetInitializationNodes.h
--- a/Engine/Plugins/Experimental/ChaosFlesh/Source/ChaosFleshNodes/Public/Dataflow/ChaosFleshPositionTargetInitializationNodes.h
+++ b/Engine/Plugins/Experimental/ChaosFlesh/Source/ChaosFleshNodes/Public/Dataflow/ChaosFleshPositionTargetInitializationNodes.h
@@ -3,20 +3,20 @@
 #pragma once 
 
 #include "CoreMinimal.h"
+#include "Dataflow/ChaosFleshNodesUtility.h"
 #include "Dataflow/DataflowCore.h"
 #include "Dataflow/DataflowEngine.h"
-#include "Dataflow/ChaosFleshKinematicInitializationNodes.h"
+
 #include "ChaosFleshPositionTargetInitializationNodes.generated.h"
 
 class USkeletalMesh;
 
-
 USTRUCT(meta = (DataflowFlesh))
 struct FAddKinematicParticlesDataflowNode : public FDataflowNode
 {
 	GENERATED_USTRUCT_BODY()
-		DATAFLOW_NODE_DEFINE_INTERNAL(FAddKinematicParticlesDataflowNode, "AddKinematicParticles", "Flesh", "")
-		DATAFLOW_NODE_RENDER_TYPE(FGeometryCollection::StaticType(), "Collection")
+	DATAFLOW_NODE_DEFINE_INTERNAL(FAddKinematicParticlesDataflowNode, "AddKinematicParticles", "Flesh", "")
+	DATAFLOW_NODE_RENDER_TYPE(FGeometryCollection::StaticType(), "Collection")
 
 public:
 	typedef FManagedArrayCollection DataType;
diff --git a/Engine/Plugins/Experimental/ChaosFlesh/Source/ChaosFleshNodes/Public/Dataflow/ChaosFleshSkeletalMeshConstraintNode.h b/Engine/Plugins/Experimental/ChaosFlesh/Source/ChaosFleshNodes/Public/Dataflow/ChaosFleshSkeletalMeshConstraintNode.h
new file mode 100644
--- /dev/null
+++ b/Engine/Plugins/Experimental/ChaosFlesh/Source/ChaosFleshNodes/Public/Dataflow/ChaosFleshSkeletalMeshConstraintNode.h
@@ -0,0 +1,43 @@
+// Copyright Epic Games, Inc. All Rights Reserved.
+#pragma once 
+
+#include "CoreMinimal.h"
+#include "Dataflow/DataflowCore.h"
+#include "Dataflow/DataflowEngine.h"
+
+#include "ChaosFleshSkeletalMeshConstraintNode.generated.h"
+
+class USkeletalMesh;
+
+// @todo(deprecate), rename to FSkeletalMeshConstraintDataflowNode
+USTRUCT(meta = (DataflowFlesh))
+struct FKinematicSkeletalMeshInitializationDataflowNode : public FDataflowNode
+{
+	GENERATED_USTRUCT_BODY()
+	DATAFLOW_NODE_DEFINE_INTERNAL(FKinematicSkeletalMeshInitializationDataflowNode, "KinematicSkeletalMeshInitialization", "Flesh", "")
+	DATAFLOW_NODE_RENDER_TYPE(FGeometryCollection::StaticType(), "Collection")
+
+public:
+	typedef FManagedArrayCollection DataType;
+
+	UPROPERTY(meta = (DataflowInput, DataflowOutput, DisplayName = "Collection"))
+		FManagedArrayCollection Collection;
+
+	UPROPERTY(meta = (DataflowInput, DisplayName = "SkeletalMesh"))
+		TObjectPtr<USkeletalMesh> SkeletalMeshIn;
+	
+	UPROPERTY(meta = (DataflowOutput, DisplayName = "SelectionSet"))
+		TArray<int32> IndicesOut;
+
+	FKinematicSkeletalMeshInitializationDataflowNode(const Dataflow::FNodeParameters& InParam, FGuid InGuid = FGuid::NewGuid())
+		: FDataflowNode(InParam, InGuid)
+	{
+		RegisterInputConnection(&Collection);
+		RegisterInputConnection(&SkeletalMeshIn);
+		RegisterOutputConnection(&Collection);
+		RegisterOutputConnection(&IndicesOut);
+	}
+
+	virtual void Evaluate(Dataflow::FContext& Context, const FDataflowOutput* Out) const override;
+};
+
diff --git a/Engine/Plugins/Experimental/ChaosFlesh/Source/ChaosFleshNodes/Public/Dataflow/ChaosFleshVertexConstraintNode.h b/Engine/Plugins/Experimental/ChaosFlesh/Source/ChaosFleshNodes/Public/Dataflow/ChaosFleshVertexConstraintNode.h
new file mode 100644
--- /dev/null
+++ b/Engine/Plugins/Experimental/ChaosFlesh/Source/ChaosFleshNodes/Public/Dataflow/ChaosFleshVertexConstraintNode.h
@@ -0,0 +1,37 @@
+// Copyright Epic Games, Inc. All Rights Reserved.
+#pragma once 
+
+#include "CoreMinimal.h"
+#include "Dataflow/DataflowCore.h"
+#include "Dataflow/DataflowEngine.h"
+
+#include "ChaosFleshVertexConstraintNode.generated.h"
+
+// @todo(deprecate), rename to FSetVerticesKinematicDataflowNode
+USTRUCT(meta = (DataflowFlesh))
+struct FSetVerticesKinematicDataflowNode : public FDataflowNode
+{
+	GENERATED_USTRUCT_BODY()
+	DATAFLOW_NODE_DEFINE_INTERNAL(FSetVerticesKinematicDataflowNode, "SetVerticesKinematic", "Flesh", "")
+
+public:
+	typedef FManagedArrayCollection DataType;
+
+	UPROPERTY(meta = (DataflowInput, DataflowOutput, DisplayName = "Collection", DataflowPassthrough  = "Collection"))
+		FManagedArrayCollection Collection;
+
+	UPROPERTY(meta = (DataflowInput, DisplayName = "SelectionSet"))
+		TArray<int32> VertexIndicesIn;
+
+	FSetVerticesKinematicDataflowNode(const Dataflow::FNodeParameters& InParam, FGuid InGuid = FGuid::NewGuid())
+		: FDataflowNode(InParam, InGuid)
+	{
+		RegisterInputConnection(&Collection);
+		RegisterOutputConnection(&Collection, &Collection);
+		RegisterInputConnection(&VertexIndicesIn);
+	}
+
+	virtual void Evaluate(Dataflow::FContext& Context, const FDataflowOutput* Out) const override;
+};
+
+
diff --git a/Engine/Plugins/Experimental/ChaosFlesh/Source/ChaosFleshNodes/Public/Dataflow/GeometryCollectionAppendCollectionTransformNode.h b/Engine/Plugins/Experimental/ChaosFlesh/Source/ChaosFleshNodes/Public/Dataflow/GeometryCollectionAppendCollectionTransformNode.h
new file mode 100644
--- /dev/null
+++ b/Engine/Plugins/Experimental/ChaosFlesh/Source/ChaosFleshNodes/Public/Dataflow/GeometryCollectionAppendCollectionTransformNode.h
@@ -0,0 +1,48 @@
+// Copyright Epic Games, Inc. All Rights Reserved.
+#pragma once 
+
+#include "CoreMinimal.h"
+#include "Dataflow/DataflowCore.h"
+#include "Dataflow/DataflowEngine.h"
+
+#include "GeometryCollectionAppendCollectionTransformNode.generated.h"
+
+//@todo(deprecate), move to GeometryCollection as AppendCollectionTransformDataflowNode
+USTRUCT(meta = (DataflowFlesh))
+struct FAppendToCollectionTransformAttributeDataflowNode : public FDataflowNode
+{
+	GENERATED_USTRUCT_BODY()
+	DATAFLOW_NODE_DEFINE_INTERNAL(FAppendToCollectionTransformAttributeDataflowNode, "AppendToCollectionTransformAttribute", "Flesh", "")
+
+public:
+	typedef FManagedArrayCollection DataType;
+
+
+	UPROPERTY(meta = (DataflowInput, DataflowOutput, DisplayName = "Collection", DataflowPassthrough = "Collection"))
+		FManagedArrayCollection Collection;
+
+	UPROPERTY(meta = (DataflowInput, DisplayName = "Transform"))
+		FTransform TransformIn = FTransform::Identity;
+
+	UPROPERTY(EditAnywhere, Category = "Dataflow")
+		FString AttributeName = FString("ComponentTransform");
+
+	UPROPERTY(EditAnywhere, Category = "Dataflow")
+		FString GroupName = FString("ComponentTransformGroup");
+
+	FAppendToCollectionTransformAttributeDataflowNode(const Dataflow::FNodeParameters& InParam, FGuid InGuid = FGuid::NewGuid())
+		: FDataflowNode(InParam, InGuid)
+	{
+		RegisterInputConnection(&TransformIn);
+		RegisterInputConnection(&Collection);
+		RegisterOutputConnection(&Collection, &Collection);
+	}
+
+	virtual void Evaluate(Dataflow::FContext& Context, const FDataflowOutput* Out) const override;
+};
+
+namespace Dataflow
+{
+	void RegisterChaosFleshKinematicInitializationNodes();
+}
+
