ChaosFlesh : Node restructuring

diff --git a/Engine/Plugins/Animation/MLDeformer/ChaosFleshGenerator/Source/ChaosFleshGenerator/Private/ChaosFleshGeneratorPrivate.cpp b/Engine/Plugins/Animation/MLDeformer/ChaosFleshGenerator/Source/ChaosFleshGenerator/Private/ChaosFleshGeneratorPrivate.cpp
--- a/Engine/Plugins/Animation/MLDeformer/ChaosFleshGenerator/Source/ChaosFleshGenerator/Private/ChaosFleshGeneratorPrivate.cpp
+++ b/Engine/Plugins/Animation/MLDeformer/ChaosFleshGenerator/Source/ChaosFleshGenerator/Private/ChaosFleshGeneratorPrivate.cpp
@@ -3,7 +3,7 @@
 #include "ChaosFleshGeneratorPrivate.h"
 
 #include "Chaos/Vector.h"
-#include "Dataflow/ChaosFleshBindingsNodes.h"
+#include "Dataflow/ChaosFleshGenerateSurfaceBindingsNode.h"
 #include "FleshGeneratorComponent.h"
 #include "FileHelpers.h"
 #include "GeometryCollection/Facades/CollectionTetrahedralBindingsFacade.h"
@@ -17,6 +17,7 @@
 #include "Rendering/SkeletalMeshRenderData.h"
 #include "Rendering/SkeletalMeshModel.h"
 
+
 DEFINE_LOG_CATEGORY(LogChaosFleshGeneratorPrivate);
 
 #define LOCTEXT_NAMESPACE "ChaosFleshGeneratorPrivate"
diff --git a/Engine/Plugins/Experimental/ChaosFlesh/ChaosFlesh.uplugin b/Engine/Plugins/Experimental/ChaosFlesh/ChaosFlesh.uplugin
--- a/Engine/Plugins/Experimental/ChaosFlesh/ChaosFlesh.uplugin
+++ b/Engine/Plugins/Experimental/ChaosFlesh/ChaosFlesh.uplugin
@@ -35,6 +35,11 @@
 			"Name": "ChaosFleshNodes",
 			"Type": "Runtime",
 			"LoadingPhase": "PreDefault"
+		},
+		{
+			"Name": "ChaosFleshDeprecatedNodes",
+			"Type": "Runtime",
+			"LoadingPhase": "PreDefault"
 		}
 	],
 	"Plugins": [
diff --git a/Engine/Plugins/Experimental/ChaosFlesh/Source/ChaosFleshDeprecatedNodes/ChaosFleshDeprecatedNodes.Build.cs b/Engine/Plugins/Experimental/ChaosFlesh/Source/ChaosFleshDeprecatedNodes/ChaosFleshDeprecatedNodes.Build.cs
new file mode 100644
--- /dev/null
+++ b/Engine/Plugins/Experimental/ChaosFlesh/Source/ChaosFleshDeprecatedNodes/ChaosFleshDeprecatedNodes.Build.cs
@@ -0,0 +1,32 @@
+// Copyright Epic Games, Inc. All Rights Reserved.
+
+namespace UnrealBuildTool.Rules
+{
+	public class ChaosFleshDeprecatedNodes : ModuleRules
+	{
+        public ChaosFleshDeprecatedNodes(ReadOnlyTargetRules Target) : base(Target)
+		{
+			PublicDependencyModuleNames.AddRange(
+				new string[]
+				{
+					"Core",
+					"CoreUObject",
+					"ChaosCore",
+					"Chaos",
+					"DataflowCore",
+					"DataflowEngine",
+					"ChaosFlesh",
+					"ChaosFleshEngine",
+					"Engine",
+					"GeometryAlgorithms",
+					"GeometryCore",
+					"MeshConversion",
+					"MeshConversionEngineTypes",
+					"MeshDescription",
+					"TetMeshing",
+					"ChaosFleshNodes"
+				}
+			);
+		}
+	}
+}
diff --git a/Engine/Plugins/Experimental/ChaosFlesh/Source/ChaosFleshDeprecatedNodes/Private/Dataflow/ChaosFleshConstructTetGridNode.cpp b/Engine/Plugins/Experimental/ChaosFlesh/Source/ChaosFleshDeprecatedNodes/Private/Dataflow/ChaosFleshConstructTetGridNode.cpp
new file mode 100644
--- /dev/null
+++ b/Engine/Plugins/Experimental/ChaosFlesh/Source/ChaosFleshDeprecatedNodes/Private/Dataflow/ChaosFleshConstructTetGridNode.cpp
@@ -0,0 +1,40 @@
+// Copyright Epic Games, Inc. All Rights Reserved.
+
+#include "Dataflow/ChaosFleshConstructTetGridNode.h"
+
+#include "ChaosFlesh/FleshCollection.h"
+#include "ChaosFlesh/ChaosFlesh.h"
+#include "ChaosLog.h"
+#include "Dataflow/ChaosFleshNodesUtility.h"
+
+
+//=============================================================================
+// FConstructTetGridNode
+//=============================================================================
+
+void FConstructTetGridNode::Evaluate(Dataflow::FContext& Context, const FDataflowOutput* Out) const
+{
+	if (Out->IsA<FManagedArrayCollection>(&Collection))
+	{
+		TUniquePtr<FFleshCollection> InCollection(GetValue<DataType>(Context, &Collection).NewCopy<FFleshCollection>());
+
+		Chaos::TVector<int32, 3> Counts(GridCellCount[0], GridCellCount[1], GridCellCount[2]);
+
+		Chaos::TVector<double, 3> MinCorner = -.5 * GridDomain;
+		Chaos::TVector<double, 3> MaxCorner = .5 * GridDomain;
+		Chaos::TUniformGrid<double, 3> Grid(MinCorner, MaxCorner, Counts, 0);
+
+		TArray<FIntVector4> Tets;
+		TArray<FVector> X;
+		Chaos::Utilities::TetMeshFromGrid<double>(Grid, Tets, X);
+
+		UE_LOG(LogChaosFlesh, Display, TEXT("TetGrid generated %d points and %d tetrahedra."), X.Num(), Tets.Num());
+
+		TArray<FIntVector3> Tris = Dataflow::GetSurfaceTriangles(Tets, !bDiscardInteriorTriangles);
+		TUniquePtr<FTetrahedralCollection> TetCollection(
+			FTetrahedralCollection::NewTetrahedralCollection(X, Tris, Tets));
+		InCollection->AppendGeometry(*TetCollection.Get());
+
+		SetValue<const DataType&>(Context, *InCollection, &Collection);
+	}
+}
diff --git a/Engine/Plugins/Experimental/ChaosFlesh/Source/ChaosFleshDeprecatedNodes/Private/Dataflow/ChaosFleshCreateTetrahedralCollectionNode.cpp b/Engine/Plugins/Experimental/ChaosFlesh/Source/ChaosFleshDeprecatedNodes/Private/Dataflow/ChaosFleshCreateTetrahedralCollectionNode.cpp
new file mode 100644
--- /dev/null
+++ b/Engine/Plugins/Experimental/ChaosFlesh/Source/ChaosFleshDeprecatedNodes/Private/Dataflow/ChaosFleshCreateTetrahedralCollectionNode.cpp
@@ -0,0 +1,258 @@
+// Copyright Epic Games, Inc. All Rights Reserved.
+
+#include "Dataflow/ChaosFleshCreateTetrahedralCollectionNode.h"
+
+#include "ChaosFlesh/ChaosFlesh.h"
+#include "ChaosFlesh/FleshCollection.h"
+#include "ChaosFlesh/FleshCollectionUtility.h"
+#include "ChaosLog.h"
+#include "DynamicMesh/DynamicMesh3.h"
+#include "DynamicMesh/DynamicMeshAABBTree3.h"
+#include "Engine/SkeletalMesh.h"
+#include "Engine/StaticMesh.h"
+#include "FTetWildWrapper.h"
+#include "Generate/IsosurfaceStuffing.h"
+#include "GeometryCollection/GeometryCollectionAlgo.h"
+#include "MeshDescription.h"
+#include "MeshDescriptionToDynamicMesh.h"
+#include "Rendering/SkeletalMeshLODImporterData.h"
+#include "Rendering/SkeletalMeshModel.h"
+#include "SkeletalMeshLODModelToDynamicMesh.h" // MeshModelingBlueprints
+#include "Spatial/FastWinding.h"
+#include "Spatial/MeshAABBTree3.h"
+
+#include UE_INLINE_GENERATED_CPP_BY_NAME(ChaosFleshCreateTetrahedralCollectionNode)
+
+
+//=============================================================================
+// FGenerateTetrahedralCollectionDataflowNodes
+//=============================================================================
+
+
+void FGenerateTetrahedralCollectionDataflowNodes::Evaluate(Dataflow::FContext& Context, const FDataflowOutput* Out) const
+{
+	if (Out->IsA<DataType>(&Collection))
+	{
+		TUniquePtr<FFleshCollection> InCollection(GetValue<DataType>(Context, &Collection).NewCopy<FFleshCollection>());
+		TObjectPtr<const UStaticMesh> InStaticMesh(GetValue<TObjectPtr<const UStaticMesh>>(Context, &StaticMesh));
+		TObjectPtr<const USkeletalMesh> InSkeletalMesh(FindInput(&SkeletalMesh) ? GetValue<TObjectPtr<const USkeletalMesh>>(Context, &SkeletalMesh) : nullptr);
+
+		if (InStaticMesh || InSkeletalMesh)
+		{
+#if WITH_EDITORONLY_DATA
+			UE::Geometry::FDynamicMesh3 DynamicMesh;
+			if (InStaticMesh)
+			{
+				// make a mesh description for UE::Geometry tools
+				FMeshDescriptionToDynamicMesh GetSourceMesh;
+				bool bUsingHiResSource = InStaticMesh->IsHiResMeshDescriptionValid();
+				const FMeshDescription* UseSourceMeshDescription =
+					(bUsingHiResSource) ? InStaticMesh->GetHiResMeshDescription() : InStaticMesh->GetMeshDescription(0);
+				GetSourceMesh.Convert(UseSourceMeshDescription, DynamicMesh);
+			}
+			else if (InSkeletalMesh)
+			{
+				// Check first if we have bulk data available and non-empty.
+				constexpr int32 LODIndex = 0;
+				FMeshDescription SourceMesh;
+				if (InSkeletalMesh->HasMeshDescription(LODIndex))
+				{
+					InSkeletalMesh->CloneMeshDescription(LODIndex, SourceMesh);
+				}
+				FMeshDescriptionToDynamicMesh Converter;
+				Converter.Convert(&SourceMesh, DynamicMesh);
+			}
+
+			if (!bComputeByComponent)
+			{
+				if (Method == TetMeshingMethod::IsoStuffing)
+				{
+					EvaluateIsoStuffing(Context, InCollection, DynamicMesh);
+				}
+				else if (Method == TetMeshingMethod::TetWild)
+				{
+					EvaluateTetWild(Context, InCollection, DynamicMesh);
+				}
+				else
+				{
+					ensureMsgf(false, TEXT("FGenerateTetrahedralCollectionDataflowNodes unsupported Method."));
+				}
+			}
+			else
+			{
+				TArray<TArray<int32>> ConnectedComponents;
+				
+				TArray<FIntVector3> Faces;
+				Faces.SetNum(DynamicMesh.TriangleCount());
+				for (int32 i = 0; i < DynamicMesh.TriangleCount(); ++i)
+				{
+					Faces[i] = FIntVector3(DynamicMesh.GetTriangle(i));
+				}
+				Chaos::Utilities::FindConnectedRegions(Faces, ConnectedComponents);
+				TArray<TUniquePtr<FFleshCollection>> CollectionBuffer;
+				for (int32 i = 0; i < ConnectedComponents.Num(); i++)
+				{
+					CollectionBuffer.Add(TUniquePtr<FFleshCollection>(new FFleshCollection()));
+					//CollectionBuffer[i]->AddElement(1, FGeometryCollection::TransformGroup);
+				}
+				ParallelFor(ConnectedComponents.Num(), [&](int32 i)
+				{		
+					UE::Geometry::FDynamicMesh3 ComponentDynamicMesh;
+					TArray<FIntVector3> ComponentFaces;
+					ComponentFaces.SetNum(ConnectedComponents[i].Num());
+					for (FVector V : DynamicMesh.VerticesItr())
+					{
+						ComponentDynamicMesh.AppendVertex(V);
+					}
+					for (int32 j = 0; j < ConnectedComponents[i].Num(); j++)
+					{
+						int32 ElementIndex = ConnectedComponents[i][j];
+						for (int32 ie = 0; ie < 3; ie++)
+						{
+							ComponentDynamicMesh.AppendTriangle(Faces[ElementIndex][0], Faces[ElementIndex][1], Faces[ElementIndex][2]);
+						}
+					}
+					ComponentDynamicMesh.CompactInPlace();
+					if (Method == TetMeshingMethod::IsoStuffing)
+					{
+						EvaluateIsoStuffing(Context, CollectionBuffer[i], ComponentDynamicMesh);
+					}
+					else if (Method == TetMeshingMethod::TetWild)
+					{
+						EvaluateTetWild(Context, CollectionBuffer[i], ComponentDynamicMesh);
+					}
+				});
+				for (int32 i = 0; i < ConnectedComponents.Num(); i++)
+				{
+					if (CollectionBuffer[i]->NumElements(FGeometryCollection::VerticesGroup))
+					{
+						TSet<int32> VertexToDeleteSet;
+						GeometryCollectionAlgo::ComputeStaleVertices(CollectionBuffer[i].Get(), VertexToDeleteSet);
+						TArray<int32> SortedVertices = VertexToDeleteSet.Array(); SortedVertices.Sort();
+						if (VertexToDeleteSet.Num()) CollectionBuffer[i]->RemoveElements(FGeometryCollection::VerticesGroup, SortedVertices);
+
+						int32 GeomIndex = InCollection->AppendGeometry(*CollectionBuffer[i].Get());
+					}
+				}
+			}
+#else
+			ensureMsgf(false, TEXT("FGenerateTetrahedralCollectionDataflowNodes is an editor only node."));
+#endif
+		} // end if InStaticMesh || InSkeletalMesh
+		SetValue<const DataType&>(Context, *InCollection, &Collection);
+	}
+}
+
+void FGenerateTetrahedralCollectionDataflowNodes::EvaluateIsoStuffing(
+	Dataflow::FContext& Context, 
+	TUniquePtr<FFleshCollection>& InCollection,
+	const UE::Geometry::FDynamicMesh3& DynamicMesh) const
+{
+#if WITH_EDITORONLY_DATA
+	if (NumCells > 0 && (-.5 <= OffsetPercent && OffsetPercent <= 0.5))
+	{
+		// Tet mesh generation
+		UE::Geometry::TIsosurfaceStuffing<double> IsosurfaceStuffing;
+		UE::Geometry::FDynamicMeshAABBTree3 Spatial(&DynamicMesh);
+		UE::Geometry::TFastWindingTree<UE::Geometry::FDynamicMesh3> FastWinding(&Spatial);
+		UE::Geometry::FAxisAlignedBox3d Bounds = Spatial.GetBoundingBox();
+		IsosurfaceStuffing.Bounds = FBox(Bounds);
+		double CellSize = Bounds.MaxDim() / NumCells;
+		IsosurfaceStuffing.CellSize = CellSize;
+		IsosurfaceStuffing.IsoValue = .5 + OffsetPercent;
+		IsosurfaceStuffing.Implicit = [&FastWinding, &Spatial](FVector3d Pos)
+		{
+			FVector3d Nearest = Spatial.FindNearestPoint(Pos);
+			double WindingSign = FastWinding.FastWindingNumber(Pos) - .5;
+			return FVector3d::Distance(Nearest, Pos) * FMathd::SignNonZero(WindingSign);
+		};
+
+		UE_LOG(LogChaosFlesh, Display, TEXT("Generating tet mesh via IsoStuffing..."));
+		IsosurfaceStuffing.Generate();
+		if (IsosurfaceStuffing.Tets.Num() > 0)
+		{
+			TArray<FVector> Vertices; Vertices.SetNumUninitialized(IsosurfaceStuffing.Vertices.Num());
+			TArray<FIntVector4> Elements; Elements.SetNumUninitialized(IsosurfaceStuffing.Tets.Num());
+			TArray<FIntVector3> SurfaceElements = Dataflow::GetSurfaceTriangles(IsosurfaceStuffing.Tets, !bDiscardInteriorTriangles);
+
+			for (int32 Tdx = 0; Tdx < IsosurfaceStuffing.Tets.Num(); ++Tdx)
+			{
+				Elements[Tdx] = IsosurfaceStuffing.Tets[Tdx];
+			}
+			for (int32 Vdx = 0; Vdx < IsosurfaceStuffing.Vertices.Num(); ++Vdx)
+			{
+				Vertices[Vdx] = IsosurfaceStuffing.Vertices[Vdx];
+			}
+
+			TUniquePtr<FTetrahedralCollection> TetCollection(FTetrahedralCollection::NewTetrahedralCollection(Vertices, SurfaceElements, Elements));
+			InCollection->AppendGeometry(*TetCollection.Get());
+
+			UE_LOG(LogChaosFlesh, Display,
+				TEXT("Generated tet mesh via IsoStuffing, num vertices: %d num tets: %d"), Vertices.Num(), Elements.Num());
+		}
+		else
+		{
+			UE_LOG(LogChaosFlesh, Warning, TEXT("IsoStuffing produced 0 tetrahedra."));
+		}
+	}
+#else
+	ensureMsgf(false, TEXT("FGenerateTetrahedralCollectionDataflowNodes is an editor only node."));
+#endif
+}
+
+void FGenerateTetrahedralCollectionDataflowNodes::EvaluateTetWild(
+	Dataflow::FContext& Context, 
+	TUniquePtr<FFleshCollection>& InCollection,
+	const UE::Geometry::FDynamicMesh3& DynamicMesh) const
+{
+#if WITH_EDITORONLY_DATA
+	if (/* placeholder for conditions for exec */true)
+	{
+		// Pull out Vertices and Triangles
+		TArray<FVector> Verts;
+		TArray<FIntVector3> Tris;
+		for (FVector V : DynamicMesh.VerticesItr())
+		{
+			Verts.Add(V);
+		}
+		for (UE::Geometry::FIndex3i Tri : DynamicMesh.TrianglesItr())
+		{
+			Tris.Emplace(Tri.A, Tri.B, Tri.C);
+		}
+
+		// Tet mesh generation
+		UE::Geometry::FTetWild::FTetMeshParameters Params;
+		Params.bCoarsen = bCoarsen;
+		Params.bExtractManifoldBoundarySurface = bExtractManifoldBoundarySurface;
+		Params.bSkipSimplification = bSkipSimplification;
+
+		Params.EpsRel = EpsRel;
+		Params.MaxIts = MaxIterations;
+		Params.StopEnergy = StopEnergy;
+		Params.IdealEdgeLength = IdealEdgeLength;
+
+		Params.bInvertOutputTets = bInvertOutputTets;
+
+		TArray<FVector> TetVerts;
+		TArray<FIntVector4> Tets;
+		FProgressCancel Progress;
+		UE_LOG(LogChaosFlesh, Display,TEXT("Generating tet mesh via TetWild..."));
+		if (UE::Geometry::FTetWild::ComputeTetMesh(Params, Verts, Tris, TetVerts, Tets, &Progress))
+		{
+			TArray<FIntVector3> SurfaceElements = Dataflow::GetSurfaceTriangles(Tets, !bDiscardInteriorTriangles);
+			TUniquePtr<FTetrahedralCollection> TetCollection(FTetrahedralCollection::NewTetrahedralCollection(TetVerts, SurfaceElements, Tets));
+			InCollection->AppendGeometry(*TetCollection.Get());
+
+			UE_LOG(LogChaosFlesh, Display,
+				TEXT("Generated tet mesh via TetWild, num vertices: %d num tets: %d"), TetVerts.Num(), Tets.Num());
+		}
+		else
+		{
+			UE_LOG(LogChaosFlesh, Error,
+				TEXT("TetWild tetrahedral mesh generation failed."));
+		}
+	}
+#else
+	ensureMsgf(false, TEXT("FGenerateTetrahedralCollectionDataflowNodes is an editor only node."));
+#endif
+}
diff --git a/Engine/Plugins/Experimental/ChaosFlesh/Source/ChaosFleshDeprecatedNodes/Private/Dataflow/ChaosFleshDeprecatedNodesPlugin.cpp b/Engine/Plugins/Experimental/ChaosFlesh/Source/ChaosFleshDeprecatedNodes/Private/Dataflow/ChaosFleshDeprecatedNodesPlugin.cpp
new file mode 100644
--- /dev/null
+++ b/Engine/Plugins/Experimental/ChaosFlesh/Source/ChaosFleshDeprecatedNodes/Private/Dataflow/ChaosFleshDeprecatedNodesPlugin.cpp
@@ -0,0 +1,26 @@
+// Copyright Epic Games, Inc. All Rights Reserved.
+
+#include "Dataflow/ChaosFleshDeprecatedNodesPlugin.h"
+
+#include "Dataflow/ChaosFleshCreateTetrahedralCollectionNode.h"
+#include "Dataflow/ChaosFleshConstructTetGridNode.h"
+
+
+#define LOCTEXT_NAMESPACE "ChaosFleshDeprecatedNodes"
+
+
+void IChaosFleshDeprecatedNodesPlugin::StartupModule()
+{
+	DATAFLOW_NODE_REGISTER_CREATION_FACTORY(FGenerateTetrahedralCollectionDataflowNodes);
+	DATAFLOW_NODE_REGISTER_CREATION_FACTORY(FConstructTetGridNode);
+}
+
+void IChaosFleshDeprecatedNodesPlugin::ShutdownModule()
+{
+}
+
+
+IMPLEMENT_MODULE(IChaosFleshDeprecatedNodesPlugin, ChaosFleshDeprecatedNodes)
+
+
+#undef LOCTEXT_NAMESPACE
diff --git a/Engine/Plugins/Experimental/ChaosFlesh/Source/ChaosFleshDeprecatedNodes/Public/Dataflow/ChaosFleshConstructTetGridNode.h b/Engine/Plugins/Experimental/ChaosFlesh/Source/ChaosFleshDeprecatedNodes/Public/Dataflow/ChaosFleshConstructTetGridNode.h
new file mode 100644
--- /dev/null
+++ b/Engine/Plugins/Experimental/ChaosFlesh/Source/ChaosFleshDeprecatedNodes/Public/Dataflow/ChaosFleshConstructTetGridNode.h
@@ -0,0 +1,48 @@
+// Copyright Epic Games, Inc. All Rights Reserved.
+
+#pragma once 
+
+#include "CoreMinimal.h"
+#include "Dataflow/DataflowCore.h"
+#include "Dataflow/DataflowEngine.h"
+#include "GeometryCollection/ManagedArrayCollection.h"
+
+#include "ChaosFleshConstructTetGridNode.generated.h"
+
+
+USTRUCT(meta = (DataflowFlesh))
+struct FConstructTetGridNode : public FDataflowNode
+{
+	GENERATED_USTRUCT_BODY()
+	DATAFLOW_NODE_DEFINE_INTERNAL(FConstructTetGridNode, "TetGrid", "Flesh", "")
+	DATAFLOW_NODE_RENDER_TYPE(FGeometryCollection::StaticType(), "Collection")
+
+public:
+	typedef FManagedArrayCollection DataType;
+
+	UPROPERTY(meta = (DataflowInput, DataflowOutput, DisplayName = "Collection"))
+	FManagedArrayCollection Collection;
+
+	UPROPERTY(EditAnywhere, Category = "Dataflow", meta = (ClampMin = "1"))
+	FIntVector GridCellCount = FIntVector(10, 10, 10);
+
+	UPROPERTY(EditAnywhere, Category = "Dataflow", meta = (ClampMin = "0.0"))
+	FVector GridDomain = FVector(10.0, 10.0, 10.0);
+
+	UPROPERTY(EditAnywhere, Category = "Dataflow")
+	bool bDiscardInteriorTriangles = true;
+
+	FConstructTetGridNode(const Dataflow::FNodeParameters& InParam, FGuid InGuid = FGuid::NewGuid())
+		: FDataflowNode(InParam, InGuid)
+	{
+		RegisterInputConnection(&Collection);
+		RegisterOutputConnection(&Collection);
+	}
+
+	virtual void Evaluate(Dataflow::FContext& Context, const FDataflowOutput* Out) const override;
+};
+
+
+
+
+
diff --git a/Engine/Plugins/Experimental/ChaosFlesh/Source/ChaosFleshDeprecatedNodes/Public/Dataflow/ChaosFleshCreateTetrahedralCollectionNode.h b/Engine/Plugins/Experimental/ChaosFlesh/Source/ChaosFleshDeprecatedNodes/Public/Dataflow/ChaosFleshCreateTetrahedralCollectionNode.h
new file mode 100644
--- /dev/null
+++ b/Engine/Plugins/Experimental/ChaosFlesh/Source/ChaosFleshDeprecatedNodes/Public/Dataflow/ChaosFleshCreateTetrahedralCollectionNode.h
@@ -0,0 +1,120 @@
+// Copyright Epic Games, Inc. All Rights Reserved.
+
+#pragma once 
+
+#include "CoreMinimal.h"
+#include "Dataflow/DataflowCore.h"
+#include "Dataflow/DataflowEngine.h"
+//#include "Dataflow/DataflowEngine.h"
+//#include "GeometryCollection/ManagedArrayCollection.h"
+#include "Dataflow/ChaosFleshNodesUtility.h"
+
+#include "ChaosFleshCreateTetrahedralCollectionNode.generated.h"
+
+class USkeletalMesh;
+class UStaticMesh;
+class FFleshCollection;
+namespace UE {
+	namespace Geometry {
+		class FDynamicMesh3;
+	}
+}
+
+
+USTRUCT(meta = (DataflowFlesh))
+struct FGenerateTetrahedralCollectionDataflowNodes : public FDataflowNode
+{
+	GENERATED_USTRUCT_BODY()
+	DATAFLOW_NODE_DEFINE_INTERNAL(FGenerateTetrahedralCollectionDataflowNodes, "GenerateTetrahedralCollection", "Flesh", "")
+	DATAFLOW_NODE_RENDER_TYPE(FGeometryCollection::StaticType(), "Collection")
+
+public:
+	typedef FManagedArrayCollection DataType;
+
+	UPROPERTY(EditAnywhere, Category = "Dataflow")
+	TEnumAsByte<TetMeshingMethod> Method = TetMeshingMethod::IsoStuffing;
+
+	//
+	// IsoStuffing
+	//
+
+	UPROPERTY(EditAnywhere, Category = "Dataflow", meta = (ClampMin = "1", EditCondition = "Method == TetMeshingMethod::IsoStuffing", EditConditionHides))
+	int32 NumCells = 32;
+
+	UPROPERTY(EditAnywhere, Category = "Dataflow", meta = (ClampMin = "-0.5", ClampMax = "0.5", EditCondition = "Method == TetMeshingMethod::IsoStuffing", EditConditionHides))
+	double OffsetPercent = 0.05;
+
+	//
+	// TetWild
+	//
+
+	//! Energy at which to stop optimizing tet quality and accept the result.
+	UPROPERTY(EditAnywhere, Category = "Dataflow", meta = (ClampMin = "0.0", EditCondition = "Method == TetMeshingMethod::TetWild", EditConditionHides))
+	double IdealEdgeLength = 1.0;
+
+	//! Maximum number of optimization iterations.
+	UPROPERTY(EditAnywhere, Category = "Dataflow", meta = (ClampMin = "1", EditCondition = "Method == TetMeshingMethod::TetWild", EditConditionHides))
+	int32 MaxIterations = 80;
+
+	//! Energy at which to stop optimizing tet quality and accept the result.
+	UPROPERTY(EditAnywhere, Category = "Dataflow", meta = (ClampMin = "0.0", EditCondition = "Method == TetMeshingMethod::TetWild", EditConditionHides))
+	double StopEnergy = 10;
+
+	//! Relative tolerance, controlling how closely the mesh must follow the input surface.
+	UPROPERTY(EditAnywhere, Category = "Dataflow", meta = (ClampMin = "0.0", EditCondition = "Method == TetMeshingMethod::TetWild", EditConditionHides))
+	double EpsRel = 1e-3;
+
+	//! Coarsen the tet mesh result.
+	UPROPERTY(EditAnywhere, Category = "Dataflow", meta = (EditCondition = "Method == TetMeshingMethod::TetWild", EditConditionHides))
+	bool bCoarsen = false;
+
+	//! Enforce that the output boundary surface should be manifold.
+	UPROPERTY(EditAnywhere, Category = "Dataflow", meta = (EditCondition = "Method == TetMeshingMethod::TetWild", EditConditionHides))
+	bool bExtractManifoldBoundarySurface = false;
+
+	//! Skip the initial simplification step.
+	UPROPERTY(EditAnywhere, Category = "Dataflow", meta = (EditCondition = "Method == TetMeshingMethod::TetWild", EditConditionHides))
+	bool bSkipSimplification = false;
+
+	//! Invert tetrahedra.
+	UPROPERTY(EditAnywhere, Category = "Dataflow", meta = (EditCondition = "Method == TetMeshingMethod::TetWild", EditConditionHides))
+	bool bInvertOutputTets = false;
+
+	//
+	// Common
+	//
+
+	UPROPERTY(EditAnywhere, Category = "Dataflow", meta = (DataflowInput, DisplayName = "StaticMesh"))
+	TObjectPtr<const UStaticMesh> StaticMesh = nullptr;
+
+	UPROPERTY(EditAnywhere, Category = "Dataflow", meta = (DataflowInput, DisplayName = "SkeletalMesh"))
+	TObjectPtr<const USkeletalMesh> SkeletalMesh = nullptr;
+
+	UPROPERTY(EditAnywhere, Category = "Dataflow")
+	bool bComputeByComponent = false;
+
+	UPROPERTY(meta = (DataflowInput, DataflowOutput, DisplayName = "Collection"))
+	FManagedArrayCollection Collection;
+
+	UPROPERTY(EditAnywhere, Category = "Dataflow")
+	bool bDiscardInteriorTriangles = true;
+
+	FGenerateTetrahedralCollectionDataflowNodes(const Dataflow::FNodeParameters& InParam, FGuid InGuid = FGuid::NewGuid())
+		: FDataflowNode(InParam, InGuid)
+	{
+		RegisterInputConnection(&StaticMesh);
+		RegisterInputConnection(&SkeletalMesh);
+		RegisterInputConnection(&Collection);
+		RegisterOutputConnection(&Collection);
+	}
+
+	virtual void Evaluate(Dataflow::FContext& Context, const FDataflowOutput* Out) const override;
+
+protected:
+	void EvaluateIsoStuffing(Dataflow::FContext& Context, TUniquePtr<FFleshCollection>& InCollection, const UE::Geometry::FDynamicMesh3& DynamicMesh) const;
+	void EvaluateTetWild(Dataflow::FContext& Context, TUniquePtr<FFleshCollection>& InCollection, const UE::Geometry::FDynamicMesh3& DynamicMesh) const;
+};
+
+
+
+
diff --git a/Engine/Plugins/Experimental/ChaosFlesh/Source/ChaosFleshDeprecatedNodes/Public/Dataflow/ChaosFleshDeprecatedNodesPlugin.h b/Engine/Plugins/Experimental/ChaosFlesh/Source/ChaosFleshDeprecatedNodes/Public/Dataflow/ChaosFleshDeprecatedNodesPlugin.h
new file mode 100644
--- /dev/null
+++ b/Engine/Plugins/Experimental/ChaosFlesh/Source/ChaosFleshDeprecatedNodes/Public/Dataflow/ChaosFleshDeprecatedNodesPlugin.h
@@ -0,0 +1,18 @@
+// Copyright Epic Games, Inc. All Rights Reserved.
+
+#pragma once
+
+#include "CoreMinimal.h"
+#include "Modules/ModuleInterface.h"
+#include "Modules/ModuleManager.h"
+
+/**
+ * The public interface to this module
+ */
+class IChaosFleshDeprecatedNodesPlugin : public IModuleInterface
+{
+public:
+	virtual void StartupModule();
+	virtual void ShutdownModule();
+};
+
diff --git a/Engine/Plugins/Experimental/ChaosFlesh/Source/ChaosFleshEditor/Private/ChaosFlesh/Cmd/ChaosFleshCommands.cpp b/Engine/Plugins/Experimental/ChaosFlesh/Source/ChaosFleshEditor/Private/ChaosFlesh/Cmd/ChaosFleshCommands.cpp
--- a/Engine/Plugins/Experimental/ChaosFlesh/Source/ChaosFleshEditor/Private/ChaosFlesh/Cmd/ChaosFleshCommands.cpp
+++ b/Engine/Plugins/Experimental/ChaosFlesh/Source/ChaosFleshEditor/Private/ChaosFlesh/Cmd/ChaosFleshCommands.cpp
@@ -13,8 +13,7 @@
 #include "Chaos/Matrix.h"
 #include "Chaos/Tetrahedron.h"
 #include "ChaosCache/FleshComponentCacheAdapter.h"
-
-#include "Dataflow/ChaosFleshBindingsNodes.h"
+#include "Dataflow/ChaosFleshGenerateSurfaceBindingsNode.h"
 
 #if USE_USD_SDK && DO_USD_CACHING
 #include "HAL/PlatformFile.h"
diff --git a/Engine/Plugins/Experimental/ChaosFlesh/Source/ChaosFleshNodes/Private/Dataflow/ChaosFleshAppendTetrahedralCollectionNode.cpp b/Engine/Plugins/Experimental/ChaosFlesh/Source/ChaosFleshNodes/Private/Dataflow/ChaosFleshAppendTetrahedralCollectionNode.cpp
new file mode 100644
--- /dev/null
+++ b/Engine/Plugins/Experimental/ChaosFlesh/Source/ChaosFleshNodes/Private/Dataflow/ChaosFleshAppendTetrahedralCollectionNode.cpp
@@ -0,0 +1,32 @@
+// Copyright Epic Games, Inc. All Rights Reserved.
+
+#include "Dataflow/ChaosFleshAppendTetrahedralCollectionNode.h"
+
+#include "ChaosFlesh/FleshCollection.h"
+#include "Dataflow/DataflowNodeFactory.h"
+
+void FAppendTetrahedralCollectionDataflowNode::Evaluate(Dataflow::FContext& Context, const FDataflowOutput* Out) const
+{
+	if (Out->IsA<FManagedArrayCollection>(&Collection1))
+	{
+		TUniquePtr<FFleshCollection> InCollection1(GetValue<FManagedArrayCollection>(Context, &Collection1).NewCopy<FFleshCollection>());
+		TUniquePtr<FFleshCollection> InCollection2(GetValue<FManagedArrayCollection>(Context, &Collection2).NewCopy<FFleshCollection>());
+		TArray<FString> GeometryGroupGuidsLocal1, GeometryGroupGuidsLocal2;
+		if (const TManagedArray<FString>* GuidArray1 = InCollection1->FindAttribute<FString>("Guid", FGeometryCollection::GeometryGroup))
+		{
+			GeometryGroupGuidsLocal1 = GuidArray1->GetConstArray();
+		}
+		if (InCollection2)
+		{
+			InCollection1->AppendGeometry(*InCollection2);
+		}		
+		if (const TManagedArray<FString>* GuidArray2 = InCollection2->FindAttribute<FString>("Guid", FGeometryCollection::GeometryGroup))
+		{
+			GeometryGroupGuidsLocal2 = GuidArray2->GetConstArray();
+		}
+		SetValue<const FManagedArrayCollection&>(Context, *InCollection1, &Collection1);
+		SetValue(Context, MoveTemp(GeometryGroupGuidsLocal1), &GeometryGroupGuidsOut1);
+		SetValue(Context, MoveTemp(GeometryGroupGuidsLocal2), &GeometryGroupGuidsOut2);
+	}
+}
+
diff --git a/Engine/Plugins/Experimental/ChaosFlesh/Source/ChaosFleshNodes/Private/Dataflow/ChaosFleshBindForRenderToSkeletalMeshNode.cpp b/Engine/Plugins/Experimental/ChaosFlesh/Source/ChaosFleshNodes/Private/Dataflow/ChaosFleshBindForRenderToSkeletalMeshNode.cpp
new file mode 100644
--- /dev/null
+++ b/Engine/Plugins/Experimental/ChaosFlesh/Source/ChaosFleshNodes/Private/Dataflow/ChaosFleshBindForRenderToSkeletalMeshNode.cpp
@@ -0,0 +1,23 @@
+// Copyright Epic Games, Inc. All Rights Reserved.
+
+#include "Dataflow/ChaosFleshBindForRenderToSkeletalMeshNode.h"
+
+#include "Dataflow/DataflowInputOutput.h"
+#include "Dataflow/DataflowNodeFactory.h"
+
+
+void FBindForRenderToSkeletalMeshDataflowNode::Evaluate(Dataflow::FContext& Context, const FDataflowOutput* Out) const
+{
+	if (Out->IsA<DataType>(&Collection))
+	{
+		DataType InCollection = GetValue<DataType>(Context, &Collection);
+
+		//
+		// @todo(dataflow) : Implemention binding for optimus deformation pipeline. 
+		//
+
+		SetValue(Context, MoveTemp(InCollection), &Collection);
+	}
+}
+
+
diff --git a/Engine/Plugins/Experimental/ChaosFlesh/Source/ChaosFleshNodes/Private/Dataflow/ChaosFleshBindingsNodes.cpp b/Engine/Plugins/Experimental/ChaosFlesh/Source/ChaosFleshNodes/Private/Dataflow/ChaosFleshBindingsNodes.cpp
deleted file mode 100644
diff --git a/Engine/Plugins/Experimental/ChaosFlesh/Source/ChaosFleshNodes/Private/Dataflow/ChaosFleshCalculateTetrehedralMetricsNode.cpp b/Engine/Plugins/Experimental/ChaosFlesh/Source/ChaosFleshNodes/Private/Dataflow/ChaosFleshCalculateTetrehedralMetricsNode.cpp
new file mode 100644
--- /dev/null
+++ b/Engine/Plugins/Experimental/ChaosFlesh/Source/ChaosFleshNodes/Private/Dataflow/ChaosFleshCalculateTetrehedralMetricsNode.cpp
@@ -0,0 +1,102 @@
+// Copyright Epic Games, Inc. All Rights Reserved.
+
+#include "Dataflow/ChaosFleshCalculateTetrehedralMetricsNode.h"
+
+#include "ChaosFlesh/ChaosFlesh.h"
+#include "Chaos/Tetrahedron.h"
+#include "ChaosFlesh/FleshCollection.h"
+#include "ChaosLog.h"
+#include "GeometryCollection/Facades/CollectionTetrahedralMetricsFacade.h"
+
+
+//=============================================================================
+// FCalculateTetMetrics
+//=============================================================================
+
+void
+FCalculateTetMetrics::Evaluate(Dataflow::FContext& Context, const FDataflowOutput* Out) const
+{
+	if (Out->IsA<DataType>(&Collection))
+	{
+		TUniquePtr<FFleshCollection> InCollection(GetValue<DataType>(Context, &Collection).NewCopy<FFleshCollection>());
+
+		TManagedArray<FIntVector4>* TetMesh =
+			InCollection->FindAttribute<FIntVector4>(
+				FTetrahedralCollection::TetrahedronAttribute, FTetrahedralCollection::TetrahedralGroup);
+		TManagedArray<int32>* TetrahedronStart =
+			InCollection->FindAttribute<int32>(
+				FTetrahedralCollection::TetrahedronStartAttribute, FGeometryCollection::GeometryGroup);
+		TManagedArray<int32>* TetrahedronCount =
+			InCollection->FindAttribute<int32>(
+				FTetrahedralCollection::TetrahedronCountAttribute, FGeometryCollection::GeometryGroup);
+
+		TManagedArray<FVector3f>* Vertex =
+			InCollection->FindAttribute<FVector3f>(
+				"Vertex", "Vertices");
+
+		GeometryCollection::Facades::FTetrahedralMetrics TetMetrics(*InCollection);
+		TManagedArrayAccessor<float>& SignedVolume = TetMetrics.GetSignedVolume();
+		TManagedArrayAccessor<float>& AspectRatio = TetMetrics.GetAspectRatio();
+
+		for (int32 TetMeshIdx = 0; TetMeshIdx < TetrahedronStart->Num(); TetMeshIdx++)
+		{
+			const int32 TetMeshStart = (*TetrahedronStart)[TetMeshIdx];
+			const int32 TetMeshCount = (*TetrahedronCount)[TetMeshIdx];
+			
+			float MinVol = TNumericLimits<float>::Max();
+			float MaxVol = -TNumericLimits<float>::Max();
+			double AvgVol = 0.0;
+
+			float MinAR = TNumericLimits<float>::Max();
+			float MaxAR = -TNumericLimits<float>::Max();
+			double AvgAR = 0.0;
+			
+			for (int32 i = 0; i < TetMeshCount; i++)
+			{
+				const int32 Idx = TetMeshStart + i;
+				const FIntVector4& Tet = (*TetMesh)[Idx];
+				Chaos::TTetrahedron<Chaos::FReal> Tetrahedron(
+					(*Vertex)[Tet[0]],
+					(*Vertex)[Tet[1]],
+					(*Vertex)[Tet[2]],
+					(*Vertex)[Tet[3]]);
+
+				float Vol = Tetrahedron.GetSignedVolume();
+				SignedVolume.ModifyAt(Idx, Vol);
+				MinVol = MinVol < Vol ? MinVol : Vol;
+				MaxVol = MaxVol < Vol ? Vol : MaxVol;
+				AvgVol += Vol;
+
+				float AR = Tetrahedron.GetAspectRatio();
+				AspectRatio.ModifyAt(Idx, AR);
+				MinAR = MinAR < AR ? MinAR : AR;
+				MaxAR = MaxAR < AR ? AR : MaxAR;
+				AvgAR += AR;
+			}
+			if (TetMeshCount)
+			{
+				AvgVol /= TetMeshCount;
+				AvgAR /= TetMeshCount;
+			}
+			else
+			{
+				MinVol = MaxVol = 0.0f;
+				MinAR = MaxAR = 0.0f;
+			}
+
+			UE_LOG(LogChaosFlesh, Display,
+				TEXT("'%s' - Tet mesh %d of %d stats:\n"
+				"    Num Tetrahedra: %d\n"
+				"    Volume (min, avg, max): %g, %g, %g\n"
+				"    Aspect ratio (min, avg, max): %g, %g, %g"),
+				*GetName().ToString(),
+				(TetMeshIdx+1), TetrahedronStart->Num(),
+				TetMeshCount,
+				MinVol, AvgVol, MaxVol,
+				MinAR, AvgAR, MaxAR);
+		}
+
+		SetValue<const DataType&>(Context, *InCollection, &Collection);
+	}
+}
+
diff --git a/Engine/Plugins/Experimental/ChaosFlesh/Source/ChaosFleshNodes/Private/Dataflow/ChaosFleshCoreNodes.cpp b/Engine/Plugins/Experimental/ChaosFlesh/Source/ChaosFleshNodes/Private/Dataflow/ChaosFleshCoreNodes.cpp
deleted file mode 100644
diff --git a/Engine/Plugins/Experimental/ChaosFlesh/Source/ChaosFleshNodes/Private/Dataflow/ChaosFleshEngineAssetNodes.cpp b/Engine/Plugins/Experimental/ChaosFlesh/Source/ChaosFleshNodes/Private/Dataflow/ChaosFleshEngineAssetNodes.cpp
--- a/Engine/Plugins/Experimental/ChaosFlesh/Source/ChaosFleshNodes/Private/Dataflow/ChaosFleshEngineAssetNodes.cpp
+++ b/Engine/Plugins/Experimental/ChaosFlesh/Source/ChaosFleshNodes/Private/Dataflow/ChaosFleshEngineAssetNodes.cpp
@@ -1,7 +1,6 @@
 // Copyright Epic Games, Inc. All Rights Reserved.
 
 #include "Dataflow/ChaosFleshEngineAssetNodes.h"
-#include "Dataflow/ChaosFleshTetrahedralNodes.h"
 
 #include "Chaos/Math/Poisson.h"
 #include "Chaos/Matrix.h"
diff --git a/Engine/Plugins/Experimental/ChaosFlesh/Source/ChaosFleshNodes/Private/Dataflow/ChaosFleshFiberDirectionInitializationNodes.cpp b/Engine/Plugins/Experimental/ChaosFlesh/Source/ChaosFleshNodes/Private/Dataflow/ChaosFleshFiberDirectionInitializationNodes.cpp
deleted file mode 100644
diff --git a/Engine/Plugins/Experimental/ChaosFlesh/Source/ChaosFleshNodes/Private/Dataflow/ChaosFleshGenerateFiberDirectionsNode.cpp b/Engine/Plugins/Experimental/ChaosFlesh/Source/ChaosFleshNodes/Private/Dataflow/ChaosFleshGenerateFiberDirectionsNode.cpp
new file mode 100644
--- /dev/null
+++ b/Engine/Plugins/Experimental/ChaosFlesh/Source/ChaosFleshNodes/Private/Dataflow/ChaosFleshGenerateFiberDirectionsNode.cpp
@@ -0,0 +1,22 @@
+// Copyright Epic Games, Inc. All Rights Reserved.
+
+#include "Dataflow/ChaosFleshGenerateFiberDirectionsNode.h"
+
+#include "Dataflow/DataflowInputOutput.h"
+#include "Dataflow/DataflowNodeFactory.h"
+
+
+void FGenerateFiberDirectionsDataflowNode::Evaluate(Dataflow::FContext& Context, const FDataflowOutput* Out) const
+{
+	if (Out->IsA<DataType>(&Collection))
+	{
+		DataType InCollection = GetValue<DataType>(Context, &Collection);
+
+		//
+		// @todo(dataflow) : Implemention fiber direction
+		//
+
+		SetValue(Context, MoveTemp(InCollection), &Collection);
+	}
+}
+
diff --git a/Engine/Plugins/Experimental/ChaosFlesh/Source/ChaosFleshNodes/Private/Dataflow/ChaosFleshGenerateSkeletalBindingsNode.cpp b/Engine/Plugins/Experimental/ChaosFlesh/Source/ChaosFleshNodes/Private/Dataflow/ChaosFleshGenerateSkeletalBindingsNode.cpp
new file mode 100644
--- /dev/null
+++ b/Engine/Plugins/Experimental/ChaosFlesh/Source/ChaosFleshNodes/Private/Dataflow/ChaosFleshGenerateSkeletalBindingsNode.cpp
@@ -0,0 +1,178 @@
+// Copyright Epic Games, Inc. All Rights Reserved.
+
+#include "Dataflow/ChaosFleshGenerateSkeletalBindingsNode.h"
+
+#include "Chaos/AABBTree.h"
+#include "Chaos/BoundingVolumeHierarchy.h"
+#include "Chaos/Tetrahedron.h"
+#include "Chaos/TriangleMesh.h"
+#include "ChaosFlesh/TetrahedralCollection.h"
+#include "Containers/Map.h"
+#include "Engine/StaticMesh.h"
+#include "Dataflow/DataflowInputOutput.h"
+#include "Engine/SkeletalMesh.h"
+#include "Dataflow/DataflowNodeFactory.h"
+#include "GeometryCollection/Facades/CollectionTetrahedralSkeletalBindingsFacade.h"
+#include "Rendering/SkeletalMeshRenderData.h"
+#include "UObject/PrimaryAssetId.h"
+
+DEFINE_LOG_CATEGORY(LogSkeletalBindings);
+
+
+namespace Dataflow
+{
+
+	TArray<int32> ChildIndices(TObjectPtr<const USkeletalMesh> SkeletalMesh, int32 StartIndex)
+	{
+		int32 IndexCount = 0;
+		TArray<int32> BranchIndices;
+
+		TArray<int32> ToProcess;
+		int32 CurrentIndex = StartIndex;
+		while (SkeletalMesh->GetRefSkeleton().IsValidIndex(CurrentIndex))
+		{
+			TArray<int32> Buffer;
+			SkeletalMesh->GetRefSkeleton().GetDirectChildBones(CurrentIndex, Buffer);
+
+			if (Buffer.Num())
+			{
+				ToProcess.Append(Buffer);
+			}
+
+			if (!IndexCount) // skip the first index to get only the chidren. 
+			{
+				BranchIndices.Add(CurrentIndex);
+			}
+
+			CurrentIndex = INDEX_NONE;
+			if (ToProcess.Num())
+			{
+				CurrentIndex = ToProcess.Pop();
+			}
+		}
+		return BranchIndices;
+	}
+
+}
+
+void
+FGenerateSkeletalBindings::Evaluate(Dataflow::FContext& Context, const FDataflowOutput* Out) const
+{
+	auto UEVertd = [](FVector3f V) { return FVector3d(V.X, V.Y, V.Z); };
+	auto UEVertf = [](FVector3d V) { return FVector3f(V.X, V.Y, V.Z); };
+
+	if (Out->IsA<DataType>(&Collection))
+	{
+		DataType OutCollection = GetValue<DataType>(Context, &Collection); // Deep copy
+
+		TManagedArray<FIntVector4>* Tetrahedron = OutCollection.FindAttribute<FIntVector4>(FTetrahedralCollection::TetrahedronAttribute, FTetrahedralCollection::TetrahedralGroup);
+		TManagedArray<int32>* TetrahedronStart =OutCollection.FindAttribute<int32>(FTetrahedralCollection::TetrahedronStartAttribute, FGeometryCollection::GeometryGroup);
+		TManagedArray<int32>* TetrahedronCount =OutCollection.FindAttribute<int32>(FTetrahedralCollection::TetrahedronCountAttribute, FGeometryCollection::GeometryGroup);
+		TManagedArray<TArray<int32>>* IncidentElements = OutCollection.FindAttribute<TArray<int32>>(FTetrahedralCollection::IncidentElementsAttribute, FGeometryCollection::VerticesGroup);
+		TManagedArray<FVector3f>* Vertex = OutCollection.FindAttribute<FVector3f>("Vertex", "Vertices");
+
+		TObjectPtr<const USkeletalMesh> SkeletalMesh = GetValue<TObjectPtr<const USkeletalMesh>>(Context, &SkeletalMeshIn);
+		if (SkeletalMesh && Tetrahedron && TetrahedronStart && TetrahedronCount && Vertex)
+		{
+			// Iterate all the tetrahedral meshs
+			for (int32 TetMeshIdx = 0; TetMeshIdx < TetrahedronStart->Num(); TetMeshIdx++)
+			{
+				const int32 TetMeshStart = (*TetrahedronStart)[TetMeshIdx];
+				const int32 TetMeshCount = (*TetrahedronCount)[TetMeshIdx];
+
+				// Build Tetrahedra
+				TArray<Chaos::TTetrahedron<Chaos::FReal>> Tets;			// Index 0 == TetMeshStart
+				TArray<Chaos::TTetrahedron<Chaos::FReal>*> BVHTetPtrs;
+				Tets.SetNumUninitialized(TetMeshCount);
+				BVHTetPtrs.SetNumUninitialized(TetMeshCount);
+				for (int32 i = 0; i < TetMeshCount; i++)
+				{
+					const int32 Idx = TetMeshStart + i;
+					const FIntVector4& Tet = (*Tetrahedron)[Idx];
+					Tets[i] = Chaos::TTetrahedron<Chaos::FReal>(
+						(*Vertex)[Tet[0]],
+						(*Vertex)[Tet[1]],
+						(*Vertex)[Tet[2]],
+						(*Vertex)[Tet[3]]);
+					BVHTetPtrs[i] = &Tets[i];
+				}
+
+				// Init BVH for tetrahedra.
+				Chaos::TBoundingVolumeHierarchy<
+					TArray<Chaos::TTetrahedron<Chaos::FReal>*>,
+					TArray<int32>,
+					Chaos::FReal,
+					3> TetBVH(BVHTetPtrs);
+
+				FName SkeletalMeshName = FName(SkeletalMesh->GetName());
+				TArray<FTransform> ComponentPose;
+				Dataflow::Animation::GlobalTransforms(SkeletalMesh->GetRefSkeleton(), ComponentPose);
+				TArray<int32> Indices = Dataflow::ChildIndices(SkeletalMesh, GetValue<int32>(Context, &BoneIndexIn));
+
+				//
+				// Do intersection tests against tets
+				//
+
+				int CurrentIndex = 0;
+				TArray<int32> TetIndex; TetIndex.SetNum(Indices.Num());
+				TArray<FVector4f> Weights;		Weights.SetNum(Indices.Num());
+				TArray<FVector3f> Offsets;		Offsets.SetNum(Indices.Num());
+				TArray<int32> IndexKey;			IndexKey.SetNum(Indices.Num());
+
+				TArray<int32> TetIntersections; TetIntersections.Reserve(64);
+				for (int32 i = 0; i < Indices.Num(); i++)
+				{
+					int32 TransformIndex = Indices[i];
+
+					FVector TestPosition = ComponentPose[TransformIndex].GetTranslation();
+					TetIntersections = TetBVH.FindAllIntersections(TestPosition);
+					for (int32 j = 0; j < TetIntersections.Num(); j++)
+					{
+						const int32 TetIdx = TetIntersections[j];
+						if (!Tets[TetIdx].Outside(TestPosition, 1.0e-2)) // includes boundary
+						{
+							Chaos::TVector<Chaos::FReal, 4> WeightsD = Tets[TetIdx].GetBarycentricCoordinates(TestPosition);
+							Weights[CurrentIndex] = FVector4f(WeightsD[0], WeightsD[1], WeightsD[2], WeightsD[3]);
+							TetIndex[CurrentIndex] = TetIdx + TetMeshStart;
+							IndexKey[CurrentIndex] = TransformIndex;
+
+							// validate
+							{
+								FIntVector4 Parent = (*Tetrahedron)[TetIndex[CurrentIndex]];
+								FVector3f EmbeddedPos =
+									(*Vertex)[Parent[0]] * Weights[CurrentIndex][0] +
+									(*Vertex)[Parent[1]] * Weights[CurrentIndex][1] +
+									(*Vertex)[Parent[2]] * Weights[CurrentIndex][2] +
+									(*Vertex)[Parent[3]] * Weights[CurrentIndex][3];
+								ensureMsgf((UEVertf(TestPosition) - EmbeddedPos).SquaredLength() < 1.0, TEXT("ERROR: Skeletal Bindings mis alignment."));
+								check((UEVertf(TestPosition) - EmbeddedPos).SquaredLength() < 1.0);
+							}
+
+							CurrentIndex++;
+							break;
+						}
+					}
+				}
+				TetIndex.SetNum(CurrentIndex, EAllowShrinking::No);
+				Weights.SetNum(CurrentIndex, EAllowShrinking::No);
+				IndexKey.SetNum(CurrentIndex, EAllowShrinking::No);
+
+				UE_LOG(LogSkeletalBindings, Display,
+					   TEXT("'%s' - Generated mesh bindings between tet mesh index %d and skeletal mesh - stats:\n"
+					   "    Skeletal transforms num: %d\n"
+					   "    Transforms in tetrahedra: %d"),
+					   *(SkeletalMeshName.ToString()), TetMeshIdx, Indices.Num(), CurrentIndex);
+
+				// Store bindings in the collection
+				if (TetIndex.Num())
+				{
+					using namespace GeometryCollection::Facades;
+					GeometryCollection::Facades::FTetrahedralSkeletalBindings TetBindings(OutCollection);
+					FString MeshBindingsName = FTetrahedralSkeletalBindings::GenerateMeshGroupName(TetMeshIdx, SkeletalMeshName);
+					TetBindings.SetBindings(MeshBindingsName, TetIndex, Weights, IndexKey);
+				}
+			} // end for TetMeshIdx
+		}
+		SetValue(Context, MoveTemp(OutCollection), &Collection);
+	}
+}
diff --git a/Engine/Plugins/Experimental/ChaosFlesh/Source/ChaosFleshNodes/Private/Dataflow/ChaosFleshGenerateSurfaceBindingsNode.cpp b/Engine/Plugins/Experimental/ChaosFlesh/Source/ChaosFleshNodes/Private/Dataflow/ChaosFleshGenerateSurfaceBindingsNode.cpp
new file mode 100644
--- /dev/null
+++ b/Engine/Plugins/Experimental/ChaosFlesh/Source/ChaosFleshNodes/Private/Dataflow/ChaosFleshGenerateSurfaceBindingsNode.cpp
@@ -0,0 +1,638 @@
+// Copyright Epic Games, Inc. All Rights Reserved.
+
+#include "Dataflow/ChaosFleshGenerateSurfaceBindingsNode.h"
+
+#include "Dataflow/DataflowInputOutput.h"
+#include "Dataflow/DataflowNodeFactory.h"
+#include "Chaos/AABBTree.h"
+#include "Chaos/BoundingVolumeHierarchy.h"
+#include "Chaos/Tetrahedron.h"
+#include "Chaos/TriangleMesh.h"
+#include "ChaosFlesh/TetrahedralCollection.h"
+#include "Containers/Map.h"
+#include "Engine/StaticMesh.h"
+#include "Engine/SkeletalMesh.h"
+#include "GeometryCollection/Facades/CollectionTetrahedralBindingsFacade.h"
+#include "IndexTypes.h"
+#include "MeshDescription.h"
+#include "MeshDescriptionToDynamicMesh.h"
+#include "Rendering/SkeletalMeshLODModel.h"
+#include "Rendering/SkeletalMeshModel.h"
+#include "Rendering/SkeletalMeshRenderData.h"
+#include "UObject/PrimaryAssetId.h"
+
+DEFINE_LOG_CATEGORY(LogMeshBindings);
+
+namespace UE::TetrahedralBindingsEngineUtil
+{
+	FString GetMeshId(const USkeletalMesh* SkeletalMesh, const bool bUseImportModel)
+	{
+		FPrimaryAssetId Id = SkeletalMesh->GetPrimaryAssetId();
+		FString MeshId = Id.IsValid() ? Id.ToString() : SkeletalMesh->GetName();
+		if (bUseImportModel)
+		{
+			MeshId.Append(TEXT("_ImportModel"));
+		}
+		return MeshId;
+	}
+
+	FString GetMeshId(const UStaticMesh* StaticMesh)
+	{
+		FPrimaryAssetId Id = StaticMesh->GetPrimaryAssetId();
+		FString MeshId = Id.IsValid() ? Id.ToString() : StaticMesh->GetName();
+		return MeshId;
+	}
+}
+
+
+void
+BuildVertexToVertexAdjacencyBuffer(
+	const UE::Geometry::FDynamicMesh3 DynamicMesh,
+	TArray<TArray<uint32>>& NeighborNodes)
+{
+	NeighborNodes.SetNum(DynamicMesh.VertexCount());
+	for (int i = 0; i < DynamicMesh.TriangleCount(); i++)
+	{
+		const UE::Geometry::FIndex3i& Tri = DynamicMesh.GetTriangle(i);
+		TArray<uint32>& N0 = NeighborNodes[Tri[0]];
+		N0.AddUnique(Tri[1]);
+		N0.AddUnique(Tri[2]);
+		TArray<uint32>& N1 = NeighborNodes[Tri[1]];
+		N1.AddUnique(Tri[0]);
+		N1.AddUnique(Tri[2]);
+		TArray<uint32>& N2 = NeighborNodes[Tri[2]];
+		N2.AddUnique(Tri[0]);
+		N2.AddUnique(Tri[1]);
+	}
+}
+
+void
+BuildVertexToVertexAdjacencyBuffer(
+	const FSkeletalMeshLODRenderData& LodRenderData,
+	TArray<TArray<uint32>>& NeighborNodes)
+{
+	const FRawStaticIndexBuffer16or32Interface* IndexBuffer = LodRenderData.MultiSizeIndexContainer.GetIndexBuffer();
+	const uint32 IndexCount = IndexBuffer->Num();
+
+	const FPositionVertexBuffer& VertexBuffer = LodRenderData.StaticVertexBuffers.PositionVertexBuffer;
+	const uint32 VertexCount = LodRenderData.StaticVertexBuffers.PositionVertexBuffer.GetNumVertices();
+
+	NeighborNodes.SetNum(0); // clear, to init clean
+	NeighborNodes.SetNum(VertexCount);
+
+	int32 BaseTriangle = 0;
+	int32 BaseVertex = 0;
+	for (int32 SectionIndex = 0; SectionIndex < LodRenderData.RenderSections.Num(); ++SectionIndex)
+	{
+		FSkelMeshRenderSection const& RenderSection = LodRenderData.RenderSections[SectionIndex];
+		int32 NumTriangles = RenderSection.NumTriangles;
+		int32 NumVertices = RenderSection.NumVertices;
+
+		TArray<uint32> RedirectionArray;
+		RedirectionArray.SetNum(VertexCount);
+		TMap<FVector, int32 /*UniqueVertexIndex*/> UniqueIndexMap;
+
+		for (int32 TriangleIt = BaseTriangle; TriangleIt < BaseTriangle + NumTriangles; ++TriangleIt)
+		{
+			const uint32 V[3] =
+			{
+				IndexBuffer->Get(TriangleIt * 3 + 0),
+				IndexBuffer->Get(TriangleIt * 3 + 1),
+				IndexBuffer->Get(TriangleIt * 3 + 2)
+			};
+
+			const FVector P[3] =
+			{
+				(FVector)VertexBuffer.VertexPosition(V[0]),
+				(FVector)VertexBuffer.VertexPosition(V[1]),
+				(FVector)VertexBuffer.VertexPosition(V[2])
+			};
+
+			for (int32 i = 0; i < 3; ++i)
+			{
+				const uint32 VertexIndex = RedirectionArray[V[i]] = UniqueIndexMap.FindOrAdd(P[i], V[i]);
+				TArray<uint32>& AdjacentVertices = NeighborNodes[VertexIndex];
+				for (int32 a = 1; a < 3; ++a)
+				{
+					const uint32 AdjacentVertexIndex = V[(i + a) % 3];
+					if (VertexIndex != AdjacentVertexIndex)
+					{
+						AdjacentVertices.AddUnique(AdjacentVertexIndex);
+					}
+				}
+			}
+		}
+
+		for (int32 VertexIt = BaseVertex + 1; VertexIt < BaseVertex + NumVertices; ++VertexIt)
+		{
+			// if this vertex has a sibling we copy the data over
+			const int32 SiblingIndex = RedirectionArray[VertexIt];
+			if (SiblingIndex != VertexIt)
+			{
+				for (int32 i = 0; i < NeighborNodes[SiblingIndex].Num(); i++)
+				{
+					const uint32 OtherNode = NeighborNodes[SiblingIndex][i];
+					if (OtherNode != VertexIt)
+					{
+						NeighborNodes[VertexIt].AddUnique(OtherNode);
+					}
+				}
+			}
+		}
+
+		BaseTriangle += NumTriangles;
+		BaseVertex += NumVertices;
+	}
+}
+
+
+
+void
+FGenerateSurfaceBindings::Evaluate(Dataflow::FContext& Context, const FDataflowOutput* Out) const
+{
+	if (Out->IsA<DataType>(&Collection))
+	{
+		DataType OutCollection = GetValue<DataType>(Context, &Collection); // Deep copy
+
+		TManagedArray<FIntVector4>* Tetrahedron = 
+			OutCollection.FindAttribute<FIntVector4>(
+				FTetrahedralCollection::TetrahedronAttribute, FTetrahedralCollection::TetrahedralGroup);
+		TManagedArray<int32>* TetrahedronStart =
+			OutCollection.FindAttribute<int32>(
+				FTetrahedralCollection::TetrahedronStartAttribute, FGeometryCollection::GeometryGroup);
+		TManagedArray<int32>* TetrahedronCount =
+			OutCollection.FindAttribute<int32>(
+				FTetrahedralCollection::TetrahedronCountAttribute, FGeometryCollection::GeometryGroup);
+		TManagedArray<TArray<int32>>* IncidentElements =
+			OutCollection.FindAttribute<TArray<int32>>(
+				FTetrahedralCollection::IncidentElementsAttribute, FGeometryCollection::VerticesGroup);
+
+		TManagedArray<FIntVector>* Triangle =
+			OutCollection.FindAttribute<FIntVector>(
+				"Indices", FGeometryCollection::FacesGroup);
+		TManagedArray<int32>* FacesStart =
+			OutCollection.FindAttribute<int32>(
+				"FaceStart", FGeometryCollection::GeometryGroup);
+		TManagedArray<int32>* FacesCount =
+			OutCollection.FindAttribute<int32>(
+				"FaceCount", FGeometryCollection::GeometryGroup);
+
+		TManagedArray<FVector3f>* Vertex = 
+			OutCollection.FindAttribute<FVector3f>(
+				"Vertex", "Vertices");
+
+		TObjectPtr<const USkeletalMesh> SkeletalMesh = GetValue<TObjectPtr<const USkeletalMesh>>(Context, &SkeletalMeshIn);
+		TObjectPtr<const UStaticMesh> StaticMesh = GetValue<TObjectPtr<const UStaticMesh>>(Context, &StaticMeshIn);
+		if ((SkeletalMesh || StaticMesh) && 
+			Tetrahedron && TetrahedronStart && TetrahedronCount &&
+			Triangle && FacesStart && FacesCount &&
+			Vertex)
+		{
+			// Extract positions to bind
+			FString MeshId;
+			TArray<TArray<FVector3f>> MeshVertices;
+			TArray<TArray<TArray<uint32>>> MeshNeighborNodes;
+			if (SkeletalMesh)
+			{
+				FPrimaryAssetId Id = SkeletalMesh->GetPrimaryAssetId();
+				MeshId = UE::TetrahedralBindingsEngineUtil::GetMeshId(SkeletalMesh, bUseSkeletalMeshImportModel);
+
+				if (!bUseSkeletalMeshImportModel)
+				{
+					FSkeletalMeshRenderData* RenderData = SkeletalMesh->GetResourceForRendering();
+
+					MeshVertices.SetNum(RenderData->LODRenderData.Num());
+					MeshNeighborNodes.SetNum(RenderData->LODRenderData.Num());
+					for (int32 i = 0; i < RenderData->LODRenderData.Num(); i++)
+					{
+						FSkeletalMeshLODRenderData* LODRenderData = &RenderData->LODRenderData[i];
+						const FPositionVertexBuffer& PositionVertexBuffer =
+							LODRenderData->StaticVertexBuffers.PositionVertexBuffer;
+
+						TArray<FVector3f>& Vertices = MeshVertices[i];
+						Vertices.SetNumUninitialized(PositionVertexBuffer.GetNumVertices());
+						for (uint32 j = 0; j < PositionVertexBuffer.GetNumVertices(); j++)
+						{
+							const FVector3f& Pos = PositionVertexBuffer.VertexPosition(j);
+							Vertices[j] = Pos;
+						}
+
+						TArray<TArray<uint32>>& NeighborNodes = MeshNeighborNodes[i];
+						BuildVertexToVertexAdjacencyBuffer(*LODRenderData, NeighborNodes);
+					}
+				}
+#if WITH_EDITOR
+				else // Import Model
+				{
+					const int32 LODIndex = 0;
+					MeshVertices.SetNum(1);
+					MeshNeighborNodes.SetNum(1);
+
+					// Check first if we have bulk data available and non-empty.
+					FMeshDescription SourceMesh;
+#if WITH_EDITORONLY_DATA
+					if (SkeletalMesh->HasMeshDescription(LODIndex))
+					{
+						// @todo(brice) : Confirm correct in release 5.4
+						//FSkeletalMeshImportData SkeletalMeshImportData;
+						//SkeletalMesh->LoadLODImportedData(LODIndex, SkeletalMeshImportData);
+						//FSkeletalMeshBuildSettings Settings;
+						//SkeletalMeshImportData.GetMeshDescription(SkeletalMesh, &Settings, SourceMesh);
+						SkeletalMesh->CloneMeshDescription(LODIndex, SourceMesh);
+					}
+					else
+#endif
+					{
+						// Fall back on the LOD model directly if no bulk data exists. When we commit
+						// the mesh description, we override using the bulk data. This can happen for older
+						// skeletal meshes, from UE 4.24 and earlier.
+						const FSkeletalMeshModel* SkeletalMeshModel = SkeletalMesh->GetImportedModel();
+						if (SkeletalMeshModel && SkeletalMeshModel->LODModels.IsValidIndex(LODIndex))
+						{
+							SkeletalMeshModel->LODModels[LODIndex].GetMeshDescription(SkeletalMesh, LODIndex, SourceMesh);
+						}
+					}
+
+					UE::Geometry::FDynamicMesh3 DynamicMesh;
+					FMeshDescriptionToDynamicMesh Converter;
+					Converter.Convert(&SourceMesh, DynamicMesh);
+
+					//const FVertexArray& SourceVertices = SourceMesh.Vertices();
+					TArray<FVector3f>& Vertices = MeshVertices[LODIndex];
+					Vertices.SetNumUninitialized(DynamicMesh.VertexCount());
+					for (int32 j = 0; j < DynamicMesh.VertexCount(); j++)
+					{
+						const FVector3d& Pos = DynamicMesh.GetVertex(j);
+						Vertices[j].Set(Pos[0], Pos[1], Pos[2]);
+					}
+
+					TArray<TArray<uint32>>& NeighborNodes = MeshNeighborNodes[LODIndex];
+					BuildVertexToVertexAdjacencyBuffer(DynamicMesh, NeighborNodes);
+				}
+#endif
+			}
+			else // StaticMesh
+			{
+				MeshId = UE::TetrahedralBindingsEngineUtil::GetMeshId(StaticMesh);
+
+				const FStaticMeshRenderData* RenderData = StaticMesh->GetRenderData();
+				const int32 NumLOD = RenderData->LODResources.Num();
+				MeshVertices.SetNum(NumLOD);
+				for (int32 i = 0; i < NumLOD; i++)
+				{
+					const FStaticMeshLODResources& LODResources = RenderData->LODResources[i];
+					const FPositionVertexBuffer& PositionVertexBuffer = LODResources.VertexBuffers.PositionVertexBuffer;
+					
+					TArray<FVector3f>& Vertices = MeshVertices[i];
+					Vertices.SetNumUninitialized(PositionVertexBuffer.GetNumVertices());
+					for (uint32 j = 0; j < PositionVertexBuffer.GetNumVertices(); j++)
+					{
+						const FVector3f& Pos = PositionVertexBuffer.VertexPosition(j);
+						Vertices[j] = Pos;
+					}
+				}
+			}
+			TArray<FString> GeometryGroupGuidsLocal = GetValue<TArray<FString>>(Context, &GeometryGroupGuidsIn);
+			TManagedArray<FString>* Guids = OutCollection.FindAttribute<FString>("Guid", FGeometryCollection::GeometryGroup);
+			for (int32 TetMeshIdx = 0; TetMeshIdx < TetrahedronStart->Num(); TetMeshIdx++)
+			{
+				if (GeometryGroupGuidsLocal.Num() && Guids)
+				{
+					if (!GeometryGroupGuidsLocal.Contains((*Guids)[TetMeshIdx]))
+					{
+						continue;
+					}
+				}
+				const int32 TetMeshStart = (*TetrahedronStart)[TetMeshIdx];
+				const int32 TetMeshCount = (*TetrahedronCount)[TetMeshIdx];
+
+				// Build Tetrahedra
+				TArray<Chaos::TTetrahedron<Chaos::FReal>> Tets;			// Index 0 == TetMeshStart
+				TArray<Chaos::TTetrahedron<Chaos::FReal>*> BVHTetPtrs;
+				Tets.SetNumUninitialized(TetMeshCount);
+				BVHTetPtrs.SetNumUninitialized(TetMeshCount);
+				for (int32 i = 0; i < TetMeshCount; i++)
+				{
+					const int32 Idx = TetMeshStart + i;
+					const FIntVector4& Tet = (*Tetrahedron)[Idx];
+					Tets[i] = Chaos::TTetrahedron<Chaos::FReal>(
+						(*Vertex)[Tet[0]],
+						(*Vertex)[Tet[1]],
+						(*Vertex)[Tet[2]],
+						(*Vertex)[Tet[3]]);
+					BVHTetPtrs[i] = &Tets[i];
+				}
+
+				// Init BVH for tetrahedra.
+				Chaos::TBoundingVolumeHierarchy<
+					TArray<Chaos::TTetrahedron<Chaos::FReal>*>, 
+					TArray<int32>, 
+					Chaos::FReal, 
+					3> TetBVH(BVHTetPtrs);
+
+				//
+				// Init boundary mesh for projections.
+				//
+
+				Chaos::FTriangleMesh TetBoundaryMesh;
+				Chaos::FTriangleMesh::TBVHType<Chaos::FRealDouble> TetBoundaryBVH;
+				TArray<Chaos::FVec3> VertexD;
+				TConstArrayView<Chaos::TVec3<Chaos::FRealDouble>> VertexDView(VertexD);
+				TArray<Chaos::FVec3> PointNormals;
+				const int32 TriMeshStart = (*FacesStart)[TetMeshIdx];
+				const int32 TriMeshCount = (*FacesCount)[TetMeshIdx];
+				if (bDoSurfaceProjection)
+				{
+					if (TriMeshCount == Triangle->GetConstArray().Num())
+					{
+						TetBoundaryMesh.Init(
+							reinterpret_cast<const TArray<Chaos::TVec3<int32>>&>(Triangle->GetConstArray()), 0, -1, false);
+					}
+					else
+					{
+						TArray<Chaos::TVec3<int32>> Faces; Faces.SetNumUninitialized(TriMeshCount);
+						for (int32 i = 0; i < TriMeshCount; i++)
+						{
+							const int32 Idx = TriMeshStart + i;
+							Faces[i] = reinterpret_cast<const Chaos::TVec3<int32>&>(Triangle->GetConstArray()[Idx]);
+						}
+						TetBoundaryMesh.Init(Faces, 0, -1, false);
+					}
+
+					// Promote vertices to double because that's what FTriangleMesh wants.
+					VertexD.SetNumUninitialized(Vertex->Num());
+					for (int32 i = 0; i < VertexD.Num(); i++)
+					{
+						VertexD[i] = Chaos::FVec3((*Vertex)[i][0], (*Vertex)[i][1], (*Vertex)[i][2]);
+					}
+					VertexDView = TConstArrayView<Chaos::TVec3<Chaos::FRealDouble>>(VertexD);
+
+					PointNormals = TetBoundaryMesh.GetPointNormals(VertexDView, false, true);
+					TetBoundaryMesh.BuildBVH(VertexDView, TetBoundaryBVH);
+				}
+
+				//
+				// Do intersection tests against tets, then the surface.
+				//
+
+				TArray<TArray<FIntVector4>> Parents; Parents.SetNum(MeshVertices.Num());
+				TArray<TArray<FVector4f>> Weights;	 Weights.SetNum(MeshVertices.Num());
+				TArray<TArray<FVector3f>> Offsets;	 Offsets.SetNum(MeshVertices.Num());
+				TArray<TArray<float>> Masks;		 Masks.SetNum(MeshVertices.Num());
+				TArray<int32> Orphans;
+				int32 TetHits = 0;
+				int32 TriHits = 0;
+				int32 Adoptions = 0;
+				int32 NumOrphans = 0;
+				for (int32 LOD = 0; LOD < MeshVertices.Num(); LOD++)
+				{
+					Parents[LOD].SetNumUninitialized(MeshVertices[LOD].Num());
+					Weights[LOD].SetNumUninitialized(MeshVertices[LOD].Num());
+					Offsets[LOD].SetNumUninitialized(MeshVertices[LOD].Num());
+					Masks[LOD].SetNumUninitialized(MeshVertices[LOD].Num());
+
+					TArray<int32> TetIntersections; TetIntersections.Reserve(64);
+					for (int32 i = 0; i < MeshVertices[LOD].Num(); i++)
+					{
+						Parents[LOD][i] = FIntVector4(INDEX_NONE);
+						Weights[LOD][i] = FVector4f(0);
+						Offsets[LOD][i] = FVector3f(0);
+						Masks[LOD][i] = 0.0; // Shader does skinning for this vertex
+
+						const FVector3f& Pos = MeshVertices[LOD][i];
+						Chaos::TVec3<Chaos::FReal> PosD(Pos[0], Pos[1], Pos[2]);
+						TetIntersections = TetBVH.FindAllIntersections(PosD);
+						int32 j = 0;
+						for (; j < TetIntersections.Num(); j++)
+						{
+							const int32 TetIdx = TetIntersections[j];
+							if (!Tets[TetIdx].Outside(Pos, 1.0e-2)) // includes boundary
+							{
+								TetHits++;
+								Parents[LOD][i] = (*Tetrahedron)[TetIdx + TetMeshStart];
+								Chaos::TVector<Chaos::FReal, 4> WeightsD = Tets[TetIdx].GetBarycentricCoordinates(Pos);
+								Weights[LOD][i] = FVector4f(WeightsD[0], WeightsD[1], WeightsD[2], WeightsD[3]);
+								Offsets[LOD][i] = FVector3f(0);
+								Masks[LOD][i] = 1.0; // Shader does sim for this vertex
+
+								FVector3f EmbeddedPos =
+									(*Vertex)[Parents[LOD][i][0]] * Weights[LOD][i][0] +
+									(*Vertex)[Parents[LOD][i][1]] * Weights[LOD][i][1] +
+									(*Vertex)[Parents[LOD][i][2]] * Weights[LOD][i][2] +
+									(*Vertex)[Parents[LOD][i][3]] * Weights[LOD][i][3];
+								check((Pos - EmbeddedPos).SquaredLength() < 1.0);
+
+								break;
+							}
+						}
+						if (j == TetIntersections.Num())
+						{
+							// This vertex didn't land inside any tetrahedra. Project it to the tet boundary surface.
+							int32 TriIdx = INDEX_NONE;
+							Chaos::FVec3 TriWeights;
+							if (bDoSurfaceProjection && 
+								TetBoundaryMesh.SmoothProject(
+									TetBoundaryBVH,
+									VertexDView,
+									PointNormals, Pos,
+									TriIdx, TriWeights, SurfaceProjectionIterations))
+							{
+								TriHits++;
+								const FIntVector& Tri = (*Triangle)[TriIdx + TriMeshStart];
+								Parents[LOD][i][0] = Tri[0];
+								Parents[LOD][i][1] = Tri[1];
+								Parents[LOD][i][2] = Tri[2];
+								Parents[LOD][i][3] = INDEX_NONE;
+
+								Weights[LOD][i][0] = TriWeights[0];
+								Weights[LOD][i][1] = TriWeights[1];
+								Weights[LOD][i][2] = TriWeights[2];
+								Weights[LOD][i][3] = 0.0;
+
+								const FVector3f EmbeddedPos =
+									TriWeights[0] * Vertex->GetConstArray()[Tri[0]] +
+									TriWeights[1] * Vertex->GetConstArray()[Tri[1]] +
+									TriWeights[2] * Vertex->GetConstArray()[Tri[2]];
+								Offsets[LOD][i] = EmbeddedPos - Pos;
+
+								Masks[LOD][i] = 1.0; // Shader does sim for this vertex
+							}
+							else
+							{
+								// Despair...
+								Orphans.Add(i);
+
+								Parents[LOD][i][0] = INDEX_NONE;
+								Parents[LOD][i][1] = INDEX_NONE;
+								Parents[LOD][i][2] = INDEX_NONE;
+								Parents[LOD][i][3] = INDEX_NONE;
+
+								Weights[LOD][i][0] = 0.0;
+								Weights[LOD][i][1] = 0.0;
+								Weights[LOD][i][2] = 0.0;
+								Weights[LOD][i][3] = 0.0;
+
+								Offsets[LOD][i][0] = 0.0;
+								Offsets[LOD][i][1] = 0.0;
+								Offsets[LOD][i][2] = 0.0;
+
+								Masks[LOD][i] = 0.0; // Shader does skinning for this vertex
+							} // if !SmoothProject()
+						} // if !TetIntersections
+					} // end for all vertices
+
+					// 
+					// Advancing front orphan reparenting
+					//
+					if (!MeshNeighborNodes.IsValidIndex(LOD))
+					{
+						continue;
+					}
+					const TArray<TArray<uint32>>& NeighborNodes = MeshNeighborNodes[LOD];
+					TSet<int32> OrphanSet(Orphans);
+					while (bDoOrphanReparenting && Orphans.Num())
+					{
+						// Find the orphan with the fewest number of orphan neighbors, and the 
+						// most non-orphans in their 1 ring.
+						int32 Orphan = INDEX_NONE;
+						int32 NumOrphanNeighbors = TNumericLimits<int32>::Max();
+						int32 NumNonOrphanNeighbors = 0;
+						for (int32 i = 0; i < Orphans.Num(); i++)
+						{
+							int32 CurrOrphan = Orphans[i];
+							if (!NeighborNodes.IsValidIndex(CurrOrphan))
+							{
+								continue;
+							}
+							const TArray<uint32>& Neighbors = NeighborNodes[CurrOrphan];
+							int32 OrphanCount = 0;
+							int32 NonOrphanCount = 0;
+							for (int32 j = 0; j < Neighbors.Num(); j++)
+							{
+								if (OrphanSet.Contains(Neighbors[j]))
+								{
+									OrphanCount++;
+								}
+								else
+								{
+									NonOrphanCount++;
+								}
+							}
+							if (OrphanCount <= NumOrphanNeighbors && NonOrphanCount > NumNonOrphanNeighbors)
+							{
+								Orphan = CurrOrphan;
+								NumOrphanNeighbors = OrphanCount;
+								NumNonOrphanNeighbors = NonOrphanCount;
+							}
+						}
+						if (Orphan == INDEX_NONE)
+						{
+							// We only have orphans with no neighbors left.
+							break;
+						}
+						const FVector3f& Pos = MeshVertices[LOD][Orphan];
+						Chaos::TVec3<Chaos::FReal> PosD(Pos[0], Pos[1], Pos[2]);
+
+						// Use the parent simplices of non-orphan neighbors as test candidates.
+						Chaos::FReal CurrDist = TNumericLimits<Chaos::FReal>::Max();
+						const TArray<uint32>& Neighbors = NeighborNodes[Orphan];
+						bool FoundBinding = false;
+						for (int32 i = 0; i < Neighbors.Num(); i++)
+						{
+							const uint32 Neighbor = Neighbors[i];
+							if (OrphanSet.Contains(Neighbor))
+							{
+								continue;
+							}
+
+							const FIntVector4& P = Parents[LOD][Neighbor];
+							int32 NumValid = 0;
+							for (int32 j = 0; j < 4; j++)
+							{
+								NumValid += P[j] != INDEX_NONE ? 1 : 0;
+							}
+
+							if (NumValid == 0)
+							{
+								continue;
+							}
+							else 
+							{
+								// Find tets that share parent indices
+								for(int32 j=0; j < 4; j++)
+								{
+									const int32 ParentIdx = P[j];
+									if (IncidentElements->GetConstArray().IsValidIndex(ParentIdx))
+									{
+										const TArray<int32>& NeighborTets = (*IncidentElements)[ParentIdx];
+										for (int32 k = 0; k < NeighborTets.Num(); k++)
+										{
+											const int32 TetIdx = NeighborTets[k] - TetMeshStart;
+											if (ensure(Tets.IsValidIndex(TetIdx)))
+											{
+												const Chaos::FTetrahedron& Tet = Tets[TetIdx];
+
+												Chaos::TVec4<Chaos::FReal> W;
+												Chaos::TVec3<Chaos::FReal> EmbeddedPos = Tet.FindClosestPointAndBary(PosD, W, 1.0e-4);
+												Chaos::TVec3<Chaos::FReal> O = EmbeddedPos - PosD;
+												Chaos::FReal Dist = O.SquaredLength();
+												if (Dist < CurrDist)
+												{
+													CurrDist = Dist;
+													Parents[LOD][Orphan] = (*Tetrahedron)[TetIdx + TetMeshStart];
+													Weights[LOD][Orphan] = FVector4f(W[0], W[1], W[2], W[3]);
+													Offsets[LOD][Orphan] = FVector3f(O[0], O[1], O[2]);
+													Masks[LOD][i] = 1.0; // Shader does sim for this vertex
+													FoundBinding = true;
+												}
+											}
+										}
+									}
+								}
+							}
+						} // end for all neighbors
+
+						// Whether or not we successfully reparented, remove the orphan from the list.
+						OrphanSet.Remove(Orphan);
+						Orphans.Remove(Orphan);
+						if (FoundBinding)
+						{
+							Adoptions++;
+						}
+						else
+						{
+							NumOrphans++;
+						}
+					} // end while(Orphans)
+					NumOrphans += Orphans.Num();
+
+					//ELogVerbosity::Type Verbosity = Orphans.Num() > 0 ? ELogVerbosity::Error : ELogVerbosity::Display;
+					UE_LOG(LogMeshBindings, Display,
+						TEXT("'%s' - Generated mesh bindings between tet mesh index %d and %s mesh of '%s' LOD %d - stats:\n"
+							"    Render vertices num: %d\n"
+							"    Vertices in tetrahedra: %d\n"
+							"    Vertices bound to tet surface: %d\n"
+							"    Orphaned vertices reparented: %d\n"
+							"    Vertices orphaned: %d"),
+						*GetName().ToString(),
+						TetMeshIdx, 
+						bUseSkeletalMeshImportModel ? TEXT("import") : TEXT("render"),
+						*MeshId, LOD,
+						MeshVertices[LOD].Num(), TetHits, TriHits, Adoptions, NumOrphans);
+
+				} // end for all LOD
+
+				// Stash bindings in the geometry collection
+				GeometryCollection::Facades::FTetrahedralBindings TetBindings(OutCollection);
+				TetBindings.DefineSchema();
+				FName MeshName(*MeshId, MeshId.Len());
+				for (int32 LOD = 0; LOD < MeshVertices.Num(); LOD++)
+				{
+					TetBindings.AddBindingsGroup(TetMeshIdx, MeshName, LOD);
+					TetBindings.SetBindingsData(Parents[LOD], Weights[LOD], Offsets[LOD], Masks[LOD]);
+				}
+			} // end for TetMeshIdx
+		}
+		SetValue(Context, MoveTemp(OutCollection), &Collection);
+	}
+}
diff --git a/Engine/Plugins/Experimental/ChaosFlesh/Source/ChaosFleshNodes/Private/Dataflow/ChaosFleshNodesPlugin.cpp b/Engine/Plugins/Experimental/ChaosFlesh/Source/ChaosFleshNodes/Private/Dataflow/ChaosFleshNodesPlugin.cpp
--- a/Engine/Plugins/Experimental/ChaosFlesh/Source/ChaosFleshNodes/Private/Dataflow/ChaosFleshNodesPlugin.cpp
+++ b/Engine/Plugins/Experimental/ChaosFlesh/Source/ChaosFleshNodes/Private/Dataflow/ChaosFleshNodesPlugin.cpp
@@ -4,21 +4,16 @@
 
 
 #include "Dataflow/ChaosFleshAuthorSceneCollisionCandidatesNode.h"
-#include "Dataflow/ChaosFleshBindingsNodes.h"
+#include "Dataflow/ChaosFleshCalculateTetrehedralMetricsNode.h"
 #include "Dataflow/ChaosFleshCollisionBodyConstraintNode.h"
-#include "Dataflow/ChaosFleshCoreNodes.h"
 #include "Dataflow/ChaosFleshCreateTetrahedronNode.h"
 #include "Dataflow/ChaosFleshEngineAssetNodes.h"
-#include "Dataflow/ChaosFleshFiberDirectionInitializationNodes.h"
 #include "Dataflow/ChaosFleshImportGEO.h"
 #include "Dataflow/ChaosFleshKinematicConstraintNode.h"
 #include "Dataflow/ChaosFleshKinematicOriginInsertionInitializationNode.h"
 #include "Dataflow/ChaosFleshKinematicTetrahedralConstraintNode.h"
-#include "Dataflow/ChaosFleshRadialTetrahedronNodes.h"
-#include "Dataflow/ChaosFleshRenderInitializationNodes.h"
+#include "Dataflow/ChaosFleshRadialTetrahedronNode.h"
 #include "Dataflow/ChaosFleshPositionTargetInitializationNodes.h"
-#include "Dataflow/ChaosFleshSkeletalBindingsNode.h"
-#include "Dataflow/ChaosFleshTetrahedralNodes.h"
 #include "Dataflow/ChaosFleshTriangleMeshSimulationPropertiesNode.h"
 #include "Dataflow/ChaosFleshSkeletalMeshConstraintNode.h"
 #include "Dataflow/ChaosFleshSkinSimulationPropertiesNode.h"
@@ -27,27 +22,30 @@
 #include "Dataflow/GeometryCollectionAppendCollectionTransformNode.h"
 #include "Modules/ModuleManager.h"
 
+#include "Dataflow/ChaosFleshAppendTetrahedralCollectionNode.h"
+#include "Dataflow/ChaosFleshBindForRenderToSkeletalMeshNode.h"
+#include "Dataflow/ChaosFleshGenerateSkeletalBindingsNode.h"
+#include "Dataflow/ChaosFleshGenerateSurfaceBindingsNode.h"
+#include "Dataflow/ChaosFleshGenerateFiberDirectionsNode.h"
+
 
 #define LOCTEXT_NAMESPACE "ChaosFleshNodes"
 
 
 void IChaosFleshNodesPlugin::StartupModule()
 {
-	Dataflow::ChaosFleshBindingsNodes();
 	Dataflow::RegisterChaosFleshEngineAssetNodes();
-	Dataflow::RegisterChaosFleshCoreNodes();
-	Dataflow::ChaosFleshFiberDirectionInitializationNodes();
-	Dataflow::ChaosFleshRenderInitializationNodes();
 	Dataflow::RegisterChaosFleshPositionTargetInitializationNodes();
-	Dataflow::ChaosFleshTetrahedralNodes();
-	Dataflow::ChaosFleshSkeletalBindingsNode();
-	Dataflow::ChaosFleshRadialTetrahedronNodes();
 	Dataflow::RegisterChaosFleshImportGEONodes();
+	DATAFLOW_NODE_REGISTER_CREATION_FACTORY(FAppendTetrahedralCollectionDataflowNode);
+	DATAFLOW_NODE_REGISTER_CREATION_FACTORY(FBindForRenderToSkeletalMeshDataflowNode); // todo delete
+	DATAFLOW_NODE_REGISTER_CREATION_FACTORY(FGenerateFiberDirectionsDataflowNode); // todo delete
+	DATAFLOW_NODE_REGISTER_CREATION_FACTORY(FGenerateSkeletalBindings);
+	DATAFLOW_NODE_REGISTER_CREATION_FACTORY(FGenerateSurfaceBindings);
 	DATAFLOW_NODE_REGISTER_CREATION_FACTORY(FCreateTetrahedronDataflowNode);
 	DATAFLOW_NODE_REGISTER_CREATION_FACTORY(FTriangleMeshSimulationPropertiesDataflowNodes);
 	DATAFLOW_NODE_REGISTER_CREATION_FACTORY(FSkinSimulationPropertiesDataflowNodes);
 	DATAFLOW_NODE_REGISTER_CREATION_FACTORY(FSetFleshDefaultPropertiesNode);
-
 	DATAFLOW_NODE_REGISTER_CREATION_FACTORY(FKinematicSkeletalMeshInitializationDataflowNode);
 	DATAFLOW_NODE_REGISTER_CREATION_FACTORY(FKinematicBodySetupInitializationDataflowNode);
 	DATAFLOW_NODE_REGISTER_CREATION_FACTORY(FKinematicInitializationDataflowNode);
@@ -56,6 +54,8 @@ void IChaosFleshNodesPlugin::StartupModule()
 	DATAFLOW_NODE_REGISTER_CREATION_FACTORY(FSetVerticesKinematicDataflowNode);
 	DATAFLOW_NODE_REGISTER_CREATION_FACTORY(FAuthorSceneCollisionCandidates);
 	DATAFLOW_NODE_REGISTER_CREATION_FACTORY(FAppendToCollectionTransformAttributeDataflowNode); 
+	DATAFLOW_NODE_REGISTER_CREATION_FACTORY(FCalculateTetMetrics);
+	DATAFLOW_NODE_REGISTER_CREATION_FACTORY(FRadialTetrahedronDataflowNodes);
 }
 
 void IChaosFleshNodesPlugin::ShutdownModule()
diff --git a/Engine/Plugins/Experimental/ChaosFlesh/Source/ChaosFleshNodes/Private/Dataflow/ChaosFleshRadialTetrahedronNode.cpp b/Engine/Plugins/Experimental/ChaosFlesh/Source/ChaosFleshNodes/Private/Dataflow/ChaosFleshRadialTetrahedronNode.cpp
new file mode 100644
--- /dev/null
+++ b/Engine/Plugins/Experimental/ChaosFlesh/Source/ChaosFleshNodes/Private/Dataflow/ChaosFleshRadialTetrahedronNode.cpp
@@ -0,0 +1,45 @@
+// Copyright Epic Games, Inc. All Rights Reserved.
+
+#include "Dataflow/ChaosFleshRadialTetrahedronNode.h"
+#include "Meshing/ChaosFleshRadialMeshing.h"
+
+#include "Chaos/Deformable/Utilities.h"
+#include "ChaosFlesh/ChaosFlesh.h"
+#include "Chaos/Utilities.h"
+#include "Chaos/UniformGrid.h"
+#include "ChaosFlesh/FleshCollection.h"
+#include "ChaosFlesh/FleshCollectionUtility.h"
+#include "ChaosLog.h"
+#include "Dataflow/DataflowInputOutput.h"
+#include "Dataflow/ChaosFleshNodesUtility.h"
+#include "DynamicMesh/DynamicMesh3.h"
+#include "DynamicMesh/DynamicMeshAABBTree3.h"
+#include "Engine/StaticMesh.h"
+#include "FTetWildWrapper.h"
+#include "Generate/IsosurfaceStuffing.h"
+#include "GeometryCollection/ManagedArrayCollection.h"
+#include "MeshDescriptionToDynamicMesh.h"
+#include "Spatial/FastWinding.h"
+#include "Spatial/MeshAABBTree3.h"
+
+//=============================================================================
+// FRadialTetrahedronDataflowNodes
+//=============================================================================
+
+
+void FRadialTetrahedronDataflowNodes::Evaluate(Dataflow::FContext& Context, const FDataflowOutput* Out) const
+{
+	if (Out->IsA<DataType>(&Collection))
+	{
+		TUniquePtr<FFleshCollection> InCollection(GetValue<DataType>(Context, &Collection).NewCopy<FFleshCollection>());
+		TArray<FIntVector4> TetElements;
+		TArray<FVector> TetVertices;
+		RadialTetMesh(InnerRadius, OuterRadius, Height, RadialSample, AngularSample, VerticalSample, BulgeRatio * (OuterRadius - InnerRadius) / (Chaos::FReal)2., TetElements, TetVertices);
+		TArray<FIntVector3> SurfaceElements = Dataflow::GetSurfaceTriangles(TetElements, !bDiscardInteriorTriangles);
+		TUniquePtr<FTetrahedralCollection> RadialCollection(FTetrahedralCollection::NewTetrahedralCollection(TetVertices, SurfaceElements, TetElements));
+		InCollection->AppendGeometry(*RadialCollection.Get());
+
+		SetValue<const DataType&>(Context, *InCollection, &Collection);
+	}
+}
+
diff --git a/Engine/Plugins/Experimental/ChaosFlesh/Source/ChaosFleshNodes/Private/Dataflow/ChaosFleshRadialTetrahedronNodes.cpp b/Engine/Plugins/Experimental/ChaosFlesh/Source/ChaosFleshNodes/Private/Dataflow/ChaosFleshRadialTetrahedronNodes.cpp
deleted file mode 100644
diff --git a/Engine/Plugins/Experimental/ChaosFlesh/Source/ChaosFleshNodes/Private/Dataflow/ChaosFleshRenderInitializationNodes.cpp b/Engine/Plugins/Experimental/ChaosFlesh/Source/ChaosFleshNodes/Private/Dataflow/ChaosFleshRenderInitializationNodes.cpp
deleted file mode 100644
diff --git a/Engine/Plugins/Experimental/ChaosFlesh/Source/ChaosFleshNodes/Private/Dataflow/ChaosFleshSkeletalBindingsNode.cpp b/Engine/Plugins/Experimental/ChaosFlesh/Source/ChaosFleshNodes/Private/Dataflow/ChaosFleshSkeletalBindingsNode.cpp
deleted file mode 100644
diff --git a/Engine/Plugins/Experimental/ChaosFlesh/Source/ChaosFleshNodes/Private/Dataflow/ChaosFleshTetrahedralNodes.cpp b/Engine/Plugins/Experimental/ChaosFlesh/Source/ChaosFleshNodes/Private/Dataflow/ChaosFleshTetrahedralNodes.cpp
deleted file mode 100644
diff --git a/Engine/Plugins/Experimental/ChaosFlesh/Source/ChaosFleshNodes/Public/Dataflow/ChaosFleshAppendTetrahedralCollectionNode.h b/Engine/Plugins/Experimental/ChaosFlesh/Source/ChaosFleshNodes/Public/Dataflow/ChaosFleshAppendTetrahedralCollectionNode.h
new file mode 100644
--- /dev/null
+++ b/Engine/Plugins/Experimental/ChaosFlesh/Source/ChaosFleshNodes/Public/Dataflow/ChaosFleshAppendTetrahedralCollectionNode.h
@@ -0,0 +1,45 @@
+// Copyright Epic Games, Inc. All Rights Reserved.
+
+#pragma once
+
+#include "CoreMinimal.h"
+
+#include "Dataflow/DataflowInputOutput.h"
+#include "Dataflow/DataflowNode.h"
+#include "GeometryCollection/ManagedArray.h"
+#include "GeometryCollection/ManagedArrayCollection.h"
+#include "GeometryCollection/GeometryCollection.h"
+
+#include "ChaosFleshAppendTetrahedralCollectionNode.generated.h"
+
+USTRUCT(meta = (DataflowFlesh))
+struct FAppendTetrahedralCollectionDataflowNode : public FDataflowNode
+{
+	GENERATED_USTRUCT_BODY()
+	DATAFLOW_NODE_DEFINE_INTERNAL(FAppendTetrahedralCollectionDataflowNode, "AppendTetrahedralCollection", "Flesh", "")
+	DATAFLOW_NODE_RENDER_TYPE(FGeometryCollection::StaticType(), "Collection1")
+public:
+	UPROPERTY(meta = (DataflowInput, DataflowOutput, DisplayName = "Collection1", DataflowPassthrough = "Collection1"));
+	FManagedArrayCollection Collection1;
+
+	UPROPERTY(meta = (DataflowInput, DisplayName = "Collection2"));
+	FManagedArrayCollection Collection2;
+
+	UPROPERTY(meta = (DataflowOutput, DisplayName = "GeometryGroupIndicesOut1"))
+		TArray<FString> GeometryGroupGuidsOut1;
+
+	UPROPERTY(meta = (DataflowOutput, DisplayName = "GeometryGroupIndicesOut2"))
+		TArray<FString> GeometryGroupGuidsOut2;
+
+	FAppendTetrahedralCollectionDataflowNode(const Dataflow::FNodeParameters& InParam, FGuid InGuid = FGuid::NewGuid())
+		: FDataflowNode(InParam, InGuid)
+	{
+		RegisterInputConnection(&Collection1);
+		RegisterInputConnection(&Collection2);
+		RegisterOutputConnection(&Collection1, &Collection1);
+		RegisterOutputConnection(&GeometryGroupGuidsOut1);
+		RegisterOutputConnection(&GeometryGroupGuidsOut2);
+	}
+
+	virtual void Evaluate(Dataflow::FContext& Context, const FDataflowOutput* Out) const override;
+};
diff --git a/Engine/Plugins/Experimental/ChaosFlesh/Source/ChaosFleshNodes/Public/Dataflow/ChaosFleshBindForRenderToSkeletalMeshNode.h b/Engine/Plugins/Experimental/ChaosFlesh/Source/ChaosFleshNodes/Public/Dataflow/ChaosFleshBindForRenderToSkeletalMeshNode.h
new file mode 100644
--- /dev/null
+++ b/Engine/Plugins/Experimental/ChaosFlesh/Source/ChaosFleshNodes/Public/Dataflow/ChaosFleshBindForRenderToSkeletalMeshNode.h
@@ -0,0 +1,39 @@
+// Copyright Epic Games, Inc. All Rights Reserved.
+
+#pragma once 
+
+#include "CoreMinimal.h"
+#include "Dataflow/DataflowCore.h"
+#include "Dataflow/DataflowEngine.h"
+
+#include "ChaosFleshBindForRenderToSkeletalMeshNode.generated.h"
+
+class USkeletalMesh;
+
+
+// @todo(delete), not actually implemented, so just delete. 
+USTRUCT(meta = (DataflowFlesh))
+struct FBindForRenderToSkeletalMeshDataflowNode : public FDataflowNode
+{
+	GENERATED_USTRUCT_BODY()
+	DATAFLOW_NODE_DEFINE_INTERNAL(FBindForRenderToSkeletalMeshDataflowNode, "BindForRenderToSkeletalMesh", "Flesh", "")
+
+public:
+	typedef FManagedArrayCollection DataType;
+
+	UPROPERTY(meta = (DataflowInput, DisplayName = "SkeletalMesh"))
+	TObjectPtr<const USkeletalMesh> SkeletalMesh = nullptr;
+
+	UPROPERTY(meta = (DataflowInput, DataflowOutput, DisplayName = "Collection", DataflowPassthrough = "Collection"))
+		FManagedArrayCollection Collection;
+
+	FBindForRenderToSkeletalMeshDataflowNode(const Dataflow::FNodeParameters& InParam, FGuid InGuid = FGuid::NewGuid())
+		: FDataflowNode(InParam, InGuid)
+	{
+		RegisterInputConnection(&SkeletalMesh);
+		RegisterInputConnection(&Collection);
+		RegisterOutputConnection(&Collection, &Collection);
+	}
+
+	virtual void Evaluate(Dataflow::FContext& Context, const FDataflowOutput* Out) const override;
+};
diff --git a/Engine/Plugins/Experimental/ChaosFlesh/Source/ChaosFleshNodes/Public/Dataflow/ChaosFleshBindingsNodes.h b/Engine/Plugins/Experimental/ChaosFlesh/Source/ChaosFleshNodes/Public/Dataflow/ChaosFleshBindingsNodes.h
deleted file mode 100644
diff --git a/Engine/Plugins/Experimental/ChaosFlesh/Source/ChaosFleshNodes/Public/Dataflow/ChaosFleshCalculateTetrehedralMetricsNode.h b/Engine/Plugins/Experimental/ChaosFlesh/Source/ChaosFleshNodes/Public/Dataflow/ChaosFleshCalculateTetrehedralMetricsNode.h
new file mode 100644
--- /dev/null
+++ b/Engine/Plugins/Experimental/ChaosFlesh/Source/ChaosFleshNodes/Public/Dataflow/ChaosFleshCalculateTetrehedralMetricsNode.h
@@ -0,0 +1,38 @@
+// Copyright Epic Games, Inc. All Rights Reserved.
+
+#pragma once 
+
+#include "CoreMinimal.h"
+#include "Dataflow/DataflowCore.h"
+#include "Dataflow/DataflowEngine.h"
+#include "GeometryCollection/ManagedArrayCollection.h"
+
+#include "ChaosFleshCalculateTetrehedralMetricsNode.generated.h"
+
+
+// Generate quality metrics
+USTRUCT(meta = (DataflowFlesh))
+struct FCalculateTetMetrics : public FDataflowNode
+{
+	GENERATED_USTRUCT_BODY()
+		DATAFLOW_NODE_DEFINE_INTERNAL(FCalculateTetMetrics, "AuthorTetMetrics", "Flesh", "")
+public:
+	typedef FManagedArrayCollection DataType;
+
+	// Passthrough geometry collection. Bindings are stored as standalone groups in the \p Collection, keyed by the name of the input render mesh and all available LOD's.
+	UPROPERTY(meta = (DataflowInput, DataflowOutput, DataflowPassthrough = "Collection", DisplayName = "Collection"))
+		FManagedArrayCollection Collection;
+
+	FCalculateTetMetrics(const Dataflow::FNodeParameters& InParam, FGuid InGuid = FGuid::NewGuid())
+		: FDataflowNode(InParam, InGuid)
+	{
+		RegisterInputConnection(&Collection);
+		RegisterOutputConnection(&Collection, &Collection);
+	}
+
+	virtual void Evaluate(Dataflow::FContext& Context, const FDataflowOutput* Out) const override;
+};
+
+
+
+
diff --git a/Engine/Plugins/Experimental/ChaosFlesh/Source/ChaosFleshNodes/Public/Dataflow/ChaosFleshCoreNodes.h b/Engine/Plugins/Experimental/ChaosFlesh/Source/ChaosFleshNodes/Public/Dataflow/ChaosFleshCoreNodes.h
deleted file mode 100644
diff --git a/Engine/Plugins/Experimental/ChaosFlesh/Source/ChaosFleshNodes/Public/Dataflow/ChaosFleshFiberDirectionInitializationNodes.h b/Engine/Plugins/Experimental/ChaosFlesh/Source/ChaosFleshNodes/Public/Dataflow/ChaosFleshFiberDirectionInitializationNodes.h
deleted file mode 100644
diff --git a/Engine/Plugins/Experimental/ChaosFlesh/Source/ChaosFleshNodes/Public/Dataflow/ChaosFleshGenerateFiberDirectionsNode.h b/Engine/Plugins/Experimental/ChaosFlesh/Source/ChaosFleshNodes/Public/Dataflow/ChaosFleshGenerateFiberDirectionsNode.h
new file mode 100644
--- /dev/null
+++ b/Engine/Plugins/Experimental/ChaosFlesh/Source/ChaosFleshNodes/Public/Dataflow/ChaosFleshGenerateFiberDirectionsNode.h
@@ -0,0 +1,33 @@
+// Copyright Epic Games, Inc. All Rights Reserved.
+
+#pragma once 
+
+#include "CoreMinimal.h"
+#include "Dataflow/DataflowCore.h"
+#include "Dataflow/DataflowEngine.h"
+
+#include "ChaosFleshGenerateFiberDirectionsNode.generated.h"
+
+// @todo(delete this node)
+USTRUCT(meta = (DataflowFlesh))
+struct FGenerateFiberDirectionsDataflowNode : public FDataflowNode
+{
+	GENERATED_USTRUCT_BODY()
+	DATAFLOW_NODE_DEFINE_INTERNAL(FGenerateFiberDirectionsDataflowNode, "GenerateFiberDirections", "Flesh", "")
+	DATAFLOW_NODE_RENDER_TYPE(FGeometryCollection::StaticType(), "Collection")
+
+public:
+	typedef FManagedArrayCollection DataType;
+
+	UPROPERTY(meta = (DataflowInput, DataflowOutput, DisplayName = "Collection", DataflowPassthrough = "Collection"))
+	FManagedArrayCollection Collection;
+
+	FGenerateFiberDirectionsDataflowNode(const Dataflow::FNodeParameters& InParam, FGuid InGuid = FGuid::NewGuid())
+		: FDataflowNode(InParam, InGuid)
+	{
+		RegisterInputConnection(&Collection);
+		RegisterOutputConnection(&Collection, &Collection);
+	}
+
+	virtual void Evaluate(Dataflow::FContext& Context, const FDataflowOutput* Out) const override;
+};
\ No newline at end of file
diff --git a/Engine/Plugins/Experimental/ChaosFlesh/Source/ChaosFleshNodes/Public/Dataflow/ChaosFleshGenerateSkeletalBindingsNode.h b/Engine/Plugins/Experimental/ChaosFlesh/Source/ChaosFleshNodes/Public/Dataflow/ChaosFleshGenerateSkeletalBindingsNode.h
new file mode 100644
--- /dev/null
+++ b/Engine/Plugins/Experimental/ChaosFlesh/Source/ChaosFleshNodes/Public/Dataflow/ChaosFleshGenerateSkeletalBindingsNode.h
@@ -0,0 +1,48 @@
+// Copyright Epic Games, Inc. All Rights Reserved.
+
+#pragma once 
+
+#include "CoreMinimal.h"
+#include "Dataflow/DataflowEngine.h"
+#include "GeometryCollection/ManagedArrayCollection.h"
+
+#include "ChaosFleshGenerateSkeletalBindingsNode.generated.h"
+
+class UStaticMesh;
+class USkeletalMesh;
+
+DECLARE_LOG_CATEGORY_EXTERN(LogSkeletalBindings, Verbose, All);
+
+// Generate barycentric bindings (used by the FleshDeformer deformer graph) of a render surface to a tetrahedral mesh.
+USTRUCT(meta = (DataflowFlesh))
+struct FGenerateSkeletalBindings : public FDataflowNode
+{
+	GENERATED_USTRUCT_BODY()
+	DATAFLOW_NODE_DEFINE_INTERNAL(FGenerateSkeletalBindings, "GenerateSkeletalBindings", "Flesh", "")
+	DATAFLOW_NODE_RENDER_TYPE(FGeometryCollection::StaticType(), "Collection")
+
+public:
+	typedef FManagedArrayCollection DataType;
+
+	// Passthrough geometry collection. Bindings are stored as standalone groups in the \p Collection, keyed by the name of the input render mesh and all available LOD's.
+	UPROPERTY(meta = (DataflowInput, DataflowOutput, DisplayName = "Collection", DataflowPassthrough = "Collection"))
+	FManagedArrayCollection Collection;
+
+	UPROPERTY(meta = (DataflowInput, DisplayName = "BoneIndex"))
+	int32 BoneIndexIn = 0;
+
+	// The input mesh, whose render surface is used to generate bindings.
+	UPROPERTY(EditAnywhere, Category = "Dataflow", meta = (DataflowInput, DisplayName = "SkeletalMesh"))
+	TObjectPtr<const USkeletalMesh> SkeletalMeshIn = nullptr;
+
+	FGenerateSkeletalBindings(const Dataflow::FNodeParameters& InParam, FGuid InGuid = FGuid::NewGuid())
+		: FDataflowNode(InParam, InGuid)
+	{
+		RegisterInputConnection(&Collection);
+		RegisterOutputConnection(&Collection, &Collection);
+		RegisterInputConnection(&BoneIndexIn);
+		RegisterInputConnection(&SkeletalMeshIn);
+	}
+
+	virtual void Evaluate(Dataflow::FContext& Context, const FDataflowOutput* Out) const override;
+};
\ No newline at end of file
diff --git a/Engine/Plugins/Experimental/ChaosFlesh/Source/ChaosFleshNodes/Public/Dataflow/ChaosFleshGenerateSurfaceBindingsNode.h b/Engine/Plugins/Experimental/ChaosFlesh/Source/ChaosFleshNodes/Public/Dataflow/ChaosFleshGenerateSurfaceBindingsNode.h
new file mode 100644
--- /dev/null
+++ b/Engine/Plugins/Experimental/ChaosFlesh/Source/ChaosFleshNodes/Public/Dataflow/ChaosFleshGenerateSurfaceBindingsNode.h
@@ -0,0 +1,78 @@
+// Copyright Epic Games, Inc. All Rights Reserved.
+
+#pragma once 
+
+#include "CoreMinimal.h"
+#include "Dataflow/DataflowEngine.h"
+#include "GeometryCollection/ManagedArrayCollection.h"
+
+#include "Engine/StaticMesh.h"
+#include "Engine/SkeletalMesh.h"
+
+#include "ChaosFleshGenerateSurfaceBindingsNode.generated.h"
+
+class UStaticMesh;
+class USkeletalMesh;
+
+DECLARE_LOG_CATEGORY_EXTERN(LogMeshBindings, Verbose, All);
+
+namespace UE::TetrahedralBindingsEngineUtil
+{
+	CHAOSFLESHNODES_API FString GetMeshId(const USkeletalMesh* SkeletalMesh, const bool bUseImportModel = false);
+	CHAOSFLESHNODES_API FString GetMeshId(const UStaticMesh* StaticMesh);
+}
+
+/** Generate barycentric bindings (used by the FleshDeformer deformer graph) of a render surface to a tetrahedral mesh. */
+USTRUCT(meta = (DataflowFlesh))
+struct FGenerateSurfaceBindings : public FDataflowNode
+{
+	GENERATED_USTRUCT_BODY()
+	DATAFLOW_NODE_DEFINE_INTERNAL(FGenerateSurfaceBindings, "GenerateSurfaceBindings", "Flesh", "")
+	DATAFLOW_NODE_RENDER_TYPE(FGeometryCollection::StaticType(), "Collection")
+
+public:
+	typedef FManagedArrayCollection DataType;
+
+	/** Passthrough geometry collection. Bindings are stored as standalone groups in the \p Collection, keyed by the name of the input render mesh and all available LOD's. */
+	UPROPERTY(meta = (DataflowInput, DataflowOutput, DataflowPassthrough = "Collection", DisplayName = "Collection"))
+	FManagedArrayCollection Collection;
+
+	/** The input mesh, whose render surface is used to generate bindings. */
+	UPROPERTY(EditAnywhere, Category = "Dataflow", meta = (DataflowInput, DisplayName = "StaticMesh"))
+	TObjectPtr<const UStaticMesh> StaticMeshIn = nullptr;
+
+	/** The input mesh, whose render surface is used to generate bindings. */
+	UPROPERTY(EditAnywhere, Category = "Dataflow", meta = (DataflowInput, DisplayName = "SkeletalMesh"))
+	TObjectPtr<const USkeletalMesh> SkeletalMeshIn = nullptr;
+
+	/** Use the import geometry of the skeletal mesh. */
+	UPROPERTY(EditAnywhere, Category = "Dataflow", meta = (DisplayName = "UseSkeletalMeshImportModel"))
+	bool bUseSkeletalMeshImportModel = false;
+
+	UPROPERTY(meta = (DataflowInput, DisplayName = "(Optional)GeometryGroupGuidsIn"))
+	TArray<FString> GeometryGroupGuidsIn;
+
+	/** Enable binding to the exterior hull of the tetrahedron mesh. */
+	UPROPERTY(EditAnywhere, Category = "Dataflow", meta = (DisplayName = "DoSurfaceProjection"))
+	bool bDoSurfaceProjection = true;
+
+	/** The maximum number of iterations to try expanding the domain while looking for surface triangles to bind to. */
+	UPROPERTY(EditAnywhere, Category = "Dataflow", meta = (DisplayName = "SurfaceProjectionIterations", EditCondition = "bDoSurfaceProjection == true"))
+	uint32 SurfaceProjectionIterations = 10;
+
+	/** When nodes aren't contained in tetrahedra and surface projection fails, try to find suitable bindings by looking to neighboring parents. */
+	UPROPERTY(EditAnywhere, Category = "Dataflow", meta = (DisplayName = "DoOrphanReparenting"))
+	bool bDoOrphanReparenting = true;
+
+	FGenerateSurfaceBindings(const Dataflow::FNodeParameters& InParam, FGuid InGuid = FGuid::NewGuid())
+		: FDataflowNode(InParam, InGuid)
+	{
+		RegisterInputConnection(&Collection);
+		RegisterOutputConnection(&Collection, &Collection);
+		RegisterInputConnection(&StaticMeshIn);
+		RegisterInputConnection(&SkeletalMeshIn);
+		RegisterInputConnection(&GeometryGroupGuidsIn);
+	}
+
+	virtual void Evaluate(Dataflow::FContext& Context, const FDataflowOutput* Out) const override;
+};
\ No newline at end of file
diff --git a/Engine/Plugins/Experimental/ChaosFlesh/Source/ChaosFleshNodes/Public/Dataflow/ChaosFleshNodesUtility.h b/Engine/Plugins/Experimental/ChaosFlesh/Source/ChaosFleshNodes/Public/Dataflow/ChaosFleshNodesUtility.h
--- a/Engine/Plugins/Experimental/ChaosFlesh/Source/ChaosFleshNodes/Public/Dataflow/ChaosFleshNodesUtility.h
+++ b/Engine/Plugins/Experimental/ChaosFlesh/Source/ChaosFleshNodes/Public/Dataflow/ChaosFleshNodesUtility.h
@@ -26,7 +26,7 @@ enum class ESkeletalSeletionMode : uint8
 
 namespace Dataflow
 {
-	TArray<FIntVector3> GetSurfaceTriangles(const TArray<FIntVector4>& Tets, const bool bKeepInterior);
+	TArray<FIntVector3> CHAOSFLESHNODES_API GetSurfaceTriangles(const TArray<FIntVector4>& Tets, const bool bKeepInterior);
 
 	/**
 	*  GetMatchingMeshIndices
@@ -34,6 +34,6 @@ namespace Dataflow
 	*	@param MeshNames	: List of name to match, or empty for all indices
 	*	@param Collection	: Collection to search
 	*/
-	TArray<int32> GetMatchingMeshIndices(const TArray<FString>& MeshNames, const FManagedArrayCollection* InCollection);
+	TArray<int32> CHAOSFLESHNODES_API GetMatchingMeshIndices(const TArray<FString>& MeshNames, const FManagedArrayCollection* InCollection);
 
 }
diff --git a/Engine/Plugins/Experimental/ChaosFlesh/Source/ChaosFleshNodes/Public/Dataflow/ChaosFleshRadialTetrahedronNode.h b/Engine/Plugins/Experimental/ChaosFlesh/Source/ChaosFleshNodes/Public/Dataflow/ChaosFleshRadialTetrahedronNode.h
new file mode 100644
--- /dev/null
+++ b/Engine/Plugins/Experimental/ChaosFlesh/Source/ChaosFleshNodes/Public/Dataflow/ChaosFleshRadialTetrahedronNode.h
@@ -0,0 +1,65 @@
+// Copyright Epic Games, Inc. All Rights Reserved.
+
+#pragma once 
+
+#include "CoreMinimal.h"
+#include "Dataflow/DataflowEngine.h"
+#include "GeometryCollection/ManagedArrayCollection.h"
+
+#include "ChaosFleshRadialTetrahedronNode.generated.h"
+
+class UStaticMesh;
+class FFleshCollection;
+
+// @todo(deprecate), rename to FRadialTetrahedronDataflowNode
+USTRUCT(meta = (DataflowFlesh))
+struct FRadialTetrahedronDataflowNodes : public FDataflowNode
+{
+	GENERATED_USTRUCT_BODY()
+	DATAFLOW_NODE_DEFINE_INTERNAL(FRadialTetrahedronDataflowNodes, "RadialTetrahedron", "Flesh", "")
+	DATAFLOW_NODE_RENDER_TYPE(FGeometryCollection::StaticType(), "Collection")
+
+public:
+	typedef FManagedArrayCollection DataType;
+
+	UPROPERTY(meta = (DataflowInput, DataflowOutput, DisplayName = "Collection"))
+	FManagedArrayCollection Collection;
+
+	UPROPERTY(EditAnywhere, Category = "Dataflow", meta = (ClampMin = "0.0"))
+	double InnerRadius = double(1.0);
+
+	UPROPERTY(EditAnywhere, Category = "Dataflow", meta = (ClampMin = "0.0"))
+	double OuterRadius = double(2.0);
+
+	UPROPERTY(EditAnywhere, Category = "Dataflow", meta = (ClampMin = "0.0"))
+	double Height = double(1.0);
+
+	UPROPERTY(EditAnywhere, Category = "Dataflow", meta = (ClampMin = "2"))
+	int32 RadialSample = 2;
+
+	UPROPERTY(EditAnywhere, Category = "Dataflow", meta = (ClampMin = "3"))
+	int32 AngularSample = 4;
+
+	UPROPERTY(EditAnywhere, Category = "Dataflow", meta = (ClampMin = "2"))
+	int32 VerticalSample = 2;
+
+	UPROPERTY(EditAnywhere, Category = "Dataflow", meta = (ClampMin = "0.0", ClampMax = "1.0"))
+	double BulgeRatio = double(0.0);
+
+	UPROPERTY(EditAnywhere, Category = "Dataflow")
+	bool bDiscardInteriorTriangles = true;
+
+	FRadialTetrahedronDataflowNodes(const Dataflow::FNodeParameters& InParam, FGuid InGuid = FGuid::NewGuid())
+		: FDataflowNode(InParam, InGuid)
+	{
+		RegisterInputConnection(&Collection);
+		RegisterOutputConnection(&Collection);
+	}
+
+	virtual void Evaluate(Dataflow::FContext& Context, const FDataflowOutput* Out) const override;
+};
+
+
+
+
+
diff --git a/Engine/Plugins/Experimental/ChaosFlesh/Source/ChaosFleshNodes/Public/Dataflow/ChaosFleshRadialTetrahedronNodes.h b/Engine/Plugins/Experimental/ChaosFlesh/Source/ChaosFleshNodes/Public/Dataflow/ChaosFleshRadialTetrahedronNodes.h
deleted file mode 100644
diff --git a/Engine/Plugins/Experimental/ChaosFlesh/Source/ChaosFleshNodes/Public/Dataflow/ChaosFleshRenderInitializationNodes.h b/Engine/Plugins/Experimental/ChaosFlesh/Source/ChaosFleshNodes/Public/Dataflow/ChaosFleshRenderInitializationNodes.h
deleted file mode 100644
diff --git a/Engine/Plugins/Experimental/ChaosFlesh/Source/ChaosFleshNodes/Public/Dataflow/ChaosFleshSkeletalBindingsNode.h b/Engine/Plugins/Experimental/ChaosFlesh/Source/ChaosFleshNodes/Public/Dataflow/ChaosFleshSkeletalBindingsNode.h
deleted file mode 100644
diff --git a/Engine/Plugins/Experimental/ChaosFlesh/Source/ChaosFleshNodes/Public/Dataflow/ChaosFleshTetrahedralNodes.h b/Engine/Plugins/Experimental/ChaosFlesh/Source/ChaosFleshNodes/Public/Dataflow/ChaosFleshTetrahedralNodes.h
deleted file mode 100644
