ChaosFlesh : Moved all the nodes into individual files.

diff --git a/Engine/Plugins/Experimental/ChaosFlesh/Source/ChaosFleshDeprecatedNodes/Private/Dataflow/ChaosFleshDeprecatedNodesPlugin.cpp b/Engine/Plugins/Experimental/ChaosFlesh/Source/ChaosFleshDeprecatedNodes/Private/Dataflow/ChaosFleshDeprecatedNodesPlugin.cpp
--- a/Engine/Plugins/Experimental/ChaosFlesh/Source/ChaosFleshDeprecatedNodes/Private/Dataflow/ChaosFleshDeprecatedNodesPlugin.cpp
+++ b/Engine/Plugins/Experimental/ChaosFlesh/Source/ChaosFleshDeprecatedNodes/Private/Dataflow/ChaosFleshDeprecatedNodesPlugin.cpp
@@ -4,6 +4,7 @@
 
 #include "Dataflow/ChaosFleshCreateTetrahedralCollectionNode.h"
 #include "Dataflow/ChaosFleshConstructTetGridNode.h"
+#include "Dataflow/ChaosFleshImportGEO.h"
 
 
 #define LOCTEXT_NAMESPACE "ChaosFleshDeprecatedNodes"
@@ -11,6 +12,7 @@
 
 void IChaosFleshDeprecatedNodesPlugin::StartupModule()
 {
+	Dataflow::RegisterChaosFleshImportGEONodes();
 	DATAFLOW_NODE_REGISTER_CREATION_FACTORY(FGenerateTetrahedralCollectionDataflowNodes);
 	DATAFLOW_NODE_REGISTER_CREATION_FACTORY(FConstructTetGridNode);
 }
diff --git a/Engine/Plugins/Experimental/ChaosFlesh/Source/ChaosFleshDeprecatedNodes/Private/Dataflow/ChaosFleshImportGEO.cpp b/Engine/Plugins/Experimental/ChaosFlesh/Source/ChaosFleshDeprecatedNodes/Private/Dataflow/ChaosFleshImportGEO.cpp
new file mode 100644
--- /dev/null
+++ b/Engine/Plugins/Experimental/ChaosFlesh/Source/ChaosFleshDeprecatedNodes/Private/Dataflow/ChaosFleshImportGEO.cpp
@@ -0,0 +1,276 @@
+// Copyright Epic Games, Inc. All Rights Reserved.
+
+#include "Dataflow/ChaosFleshImportGEO.h"
+#include "Dataflow/ChaosFleshNodesUtility.h"
+#include "ChaosFlesh/TetrahedralCollection.h"
+#include "ChaosFlesh/FleshCollection.h"
+
+#include "HAL/PlatformFile.h"
+#include "HAL/PlatformFileManager.h"
+
+#include "GEO/GEO.h"
+#include "Misc/Paths.h"
+
+DEFINE_LOG_CATEGORY(LogImportGEO);
+
+namespace Dataflow
+{
+	void RegisterChaosFleshImportGEONodes()
+	{
+		DATAFLOW_NODE_REGISTER_CREATION_FACTORY(FImportGEO);
+		DATAFLOW_NODE_REGISTER_CREATION_FACTORY(FExtractGEOInt);
+		DATAFLOW_NODE_REGISTER_CREATION_FACTORY(FExtractGEOIntVector);
+		DATAFLOW_NODE_REGISTER_CREATION_FACTORY(FExtractGEOFloatVector);
+	}
+}
+
+//=============================================================================
+// FExtractGEOInt, FExtractGEOIntVector, FExtractGEOFloatVector
+//=============================================================================
+
+
+void
+FExtractGEOInt::Evaluate(Dataflow::FContext& Context, const FDataflowOutput* Out) const
+{
+	if (Out->IsA<int32>(&Value))
+	{
+		const FGEOMapStringInt& Values = GetValue<FGEOMapStringInt>(Context, &IntVars);
+		const FString& TargetName = VarName;
+		if (Values.IntVars.Contains(TargetName))
+		{
+			int32 OutValue = *Values.IntVars.Find(TargetName);
+			SetValue<int32>(Context, MoveTemp(OutValue), &Value);
+		}
+	}
+}
+
+void
+FExtractGEOIntVector::Evaluate(Dataflow::FContext& Context, const FDataflowOutput* Out) const
+{
+	if (Out->IsA<TArray<int32>>(&Value))
+	{
+		const FGEOMapStringArrayInt& Values = GetValue<FGEOMapStringArrayInt>(Context, &IntVectorVars);
+		const FString& TargetName = VarName;
+		if (Values.IntVectorVars.Contains(TargetName))
+		{
+			TArray<int32> OutValue = *Values.IntVectorVars.Find(TargetName);
+			SetValue<TArray<int32>>(Context, MoveTemp(OutValue), &Value);
+		}
+	}
+}
+
+void
+FExtractGEOFloatVector::Evaluate(Dataflow::FContext& Context, const FDataflowOutput* Out) const
+{
+	if (Out->IsA<TArray<float>>(&Value))
+	{
+		const FGEOMapStringArrayFloat& Values = GetValue<FGEOMapStringArrayFloat>(Context, &FloatVectorVars);
+		const FString& TargetName = VarName;
+		if (Values.FloatVectorVars.Contains(TargetName))
+		{
+			TArray<float> OutValue = *Values.FloatVectorVars.Find(TargetName);
+			SetValue<TArray<float>>(Context, MoveTemp(OutValue), &Value);
+		}
+	}
+}
+
+//=============================================================================
+// FImportGEO
+//=============================================================================
+
+bool
+FImportGEO::ReadGEOFile(const bool PrintStats) const
+{
+#if WITH_EDITOR
+	if (!FPaths::FileExists(Filename.FilePath))
+	{
+		UE_LOG(LogImportGEO, Error, TEXT("'%s' cannot find GEO file: '%s'"),
+			*GetName().ToString(), *Filename.FilePath);
+		return false;
+	}
+
+	FPlatformFileManager& FileManager = FPlatformFileManager::Get();
+	IPlatformFile& PlatformFile = FileManager.GetPlatformFile();
+	if (FilePath == Filename.FilePath)
+	{
+		FDateTime CurrModTime = PlatformFile.GetTimeStamp(*FilePath);
+		if (FileModTime == CurrModTime)
+		{
+			return true;
+		}
+	}
+
+	IntVars.Empty();
+	IntVectorVars.Empty();
+	FloatVectorVars.Empty();
+	IndexedStringVars.Empty();
+
+	std::string filestr(TCHAR_TO_ANSI(*Filename.FilePath));
+	std::string errlog;
+	std::stringstream out(errlog);
+	if (!ChaosFlesh::ReadGEO(filestr, IntVars, IntVectorVars, FloatVectorVars, IndexedStringVars, &out))
+	{
+		FString Tmp(ANSI_TO_TCHAR(errlog.c_str()));
+		UE_LOG(LogImportGEO, Error, TEXT("'%s' - Failed to read GEO file '%s'; log:\n%s"),
+			*GetName().ToString(), *Filename.FilePath, *Tmp);
+		return false;
+	}
+
+	FilePath = Filename.FilePath;
+	FileModTime = PlatformFile.GetTimeStamp(*FilePath);
+
+	if (PrintStats)
+	{
+		FString OutStr = FString::Format(TEXT("'{0}' - GEO file '{1}' contents:\n"), { *GetName().ToString(), *Filename.FilePath });
+		OutStr += FString::Format(TEXT("    Int variables: ({})\n"), { IntVars.Num() });
+		{
+			int32 i = 1;
+			for (TMap<FString, int32>::TConstIterator It = IntVars.CreateConstIterator(); It; ++It, ++i)
+			{
+				OutStr += FString::Format(TEXT("        {0}: '{1}'\n"), { i, *It.Key() });
+			}
+		}
+		OutStr += FString::Format(TEXT("    Int vector variables: ({0})\n"), { IntVectorVars.Num() });
+		{
+			int32 i = 1;
+			for (TMap<FString, TArray<int32>>::TConstIterator It = IntVectorVars.CreateConstIterator(); It; ++It, ++i)
+			{
+				OutStr += FString::Format(TEXT("        {0}: '{1}', Num = {2}\n"), { i, *It.Key(), It.Value().Num() });
+			}
+		}
+		OutStr += FString::Format(TEXT("    Float vector variables: ({0})\n"), { FloatVectorVars.Num() });
+		{
+			int32 i = 1;
+			for (TMap<FString, TArray<float>>::TConstIterator It = FloatVectorVars.CreateConstIterator(); It; ++It, ++i)
+			{
+				OutStr += FString::Format(TEXT("        {0}: '{1}', Num = {2}\n"), { i, *It.Key(), It.Value().Num() });
+			}
+		}
+/*		OutStr += FString::Format(TEXT("    Indexed string variables: ({})\n"), { IndexedStringVars.Num() });
+		{
+			int32 i = 1;
+			for (TMap<FString, TPair<TArray<std::string>, TArray<int32>>>::TConstIterator It = IndexedStringVars.CreateConstIterator(); It; ++It, ++i)
+			{
+				const TPair<TArray<std::string>, TArray<int32>>& Value = It.Value();
+				OutStr += FString::Format(
+					TEXT("        {}: '{}', Num = {}\n"),
+					{ i, *It.Key(), Value.Get<0>().Num() });
+
+				const TArray<std::string>& Value0 = Value.Get<0>();
+				const TArray<int32>& Value1 = Value.Get<1>();
+				for (int32 j = 0; j < Value0.Num(); j++)
+				{
+					FString Tmp(ANSI_TO_TCHAR(Value0[j].c_str()));
+					OutStr += FString::Format(
+						TEXT("            {}: '{}', Num = {}\n"), { j, *Tmp, Value1.Num() });
+				}
+			}
+		}
+*/
+		UE_LOG(LogImportGEO, Log, TEXT("%s"), *OutStr);
+	}
+	return true;
+#else // WITH_EDITOR
+	UE_LOG(LogImportGEO, Error, TEXT("ImportGEO node '%s' is not available on this platform."), *GetName().ToString());
+	return false;
+#endif // WITH_EDITOR
+}
+
+void
+FImportGEO::Evaluate(
+	Dataflow::FContext& Context,
+	const FDataflowOutput* Out) const
+{
+	if (Out->IsA<FManagedArrayCollection>(&Collection))
+	{
+		TUniquePtr<FFleshCollection> InCollection(GetValue<FManagedArrayCollection>(Context, &Collection).NewCopy<FFleshCollection>());
+
+		if (!ReadGEOFile())
+		{
+			return;
+		}
+
+		//
+		// Tet Mesh
+		//
+
+		if (bImportTetrahedronMesh)
+		{
+			TArray<FVector> Vertices;
+			TArray<FIntVector4> Elements;
+
+			auto FvIt = FloatVectorVars.Find("position");
+			if (FvIt == nullptr)
+				FvIt = FloatVectorVars.Find("P");
+			if (FvIt != nullptr)
+			{
+				const TArray<float>& Coords = *FvIt;
+				Vertices.Reserve(Coords.Num() / 3);
+				for (size_t i = 0; i < Coords.Num(); i += 3)
+				{
+					FVector pt;
+					for (size_t j = 0; j < 3; j++)
+						pt[j] = Coords[i + j];
+					Vertices.Add(pt);
+				}
+			}
+			UE_LOG(LogImportGEO, Display, TEXT("Got %d points."), Vertices.Num());
+
+			auto IvIt = IntVectorVars.Find("pointref.indices");
+			if (IvIt != nullptr)
+			{
+				auto iIt = IntVars.Find("Tetrahedron_run:startvertex");
+				int32 StartIndex = iIt == nullptr ? 0 : *iIt;
+
+				iIt = IntVars.Find("Tetrahedron_run:nprimitives");
+				int32 NumTets = iIt == nullptr ? -1 : *iIt;
+
+				UE_LOG(LogImportGEO, Display, TEXT("Tet start index: %d num tets: %d"), StartIndex, NumTets);
+
+				const TArray<int32>& Indices = *IvIt;
+				Elements.Reserve(Indices.Num() / 4);
+				size_t stopIndex = NumTets != -1 ? StartIndex + NumTets * 4 : Indices.Num();
+				for (size_t i = StartIndex; i < stopIndex; i += 4)
+				{
+					FIntVector4 tet;
+					for (size_t j = 0; j < 4; j++)
+						tet[j] = Indices[i + j];
+					Elements.Add(tet);
+				}
+				UE_LOG(LogImportGEO, Display, TEXT("Got %d tets."), Elements.Num());
+			}
+
+			if (Vertices.Num() && Elements.Num())
+			{
+				TArray<FIntVector3> SurfaceElements = Dataflow::GetSurfaceTriangles(Elements, !bDiscardInteriorTriangles);
+				TUniquePtr<FTetrahedralCollection> OBJCollection(
+					FTetrahedralCollection::NewTetrahedralCollection(Vertices, SurfaceElements, Elements));
+				InCollection->AppendGeometry(*OBJCollection.Get());
+
+				SetValue<const FManagedArrayCollection&>(Context, *InCollection, &Collection);
+
+				UE_LOG(LogImportGEO, Log,
+					TEXT("'%s' - Authored tetrahedral mesh of %d vertices and %d tetrahedra"),
+					*GetName().ToString(), Vertices.Num(), Elements.Num());
+			}
+		}
+	}
+	else if (Out->IsA<FGEOMapStringInt>(&IntVarsOutput))
+	{
+		FGEOMapStringInt Value;
+		Value.IntVars = IntVars;
+		SetValue<const FGEOMapStringInt&>(Context, Value, &IntVarsOutput);
+	}
+	else if (Out->IsA<FGEOMapStringArrayInt>(&IntVectorVarsOutput))
+	{
+		FGEOMapStringArrayInt Value;
+		Value.IntVectorVars = IntVectorVars;
+		SetValue<const FGEOMapStringArrayInt&>(Context, Value, &IntVectorVarsOutput);
+	}
+	else if (Out->IsA<FGEOMapStringArrayFloat>(&FloatVectorVarsOutput))
+	{
+		FGEOMapStringArrayFloat Value;
+		Value.FloatVectorVars = FloatVectorVars;
+		SetValue<const FGEOMapStringArrayFloat&>(Context, Value, &FloatVectorVarsOutput);
+	}
+}
diff --git a/Engine/Plugins/Experimental/ChaosFlesh/Source/ChaosFleshDeprecatedNodes/Public/Dataflow/ChaosFleshImportGEO.h b/Engine/Plugins/Experimental/ChaosFlesh/Source/ChaosFleshDeprecatedNodes/Public/Dataflow/ChaosFleshImportGEO.h
new file mode 100644
--- /dev/null
+++ b/Engine/Plugins/Experimental/ChaosFlesh/Source/ChaosFleshDeprecatedNodes/Public/Dataflow/ChaosFleshImportGEO.h
@@ -0,0 +1,201 @@
+// Copyright Epic Games, Inc. All Rights Reserved.
+
+#pragma once
+
+#include "CoreMinimal.h"
+#include "Dataflow/DataflowEngine.h"
+#include "GeometryCollection/ManagedArrayCollection.h"
+
+#include "ChaosFleshImportGEO.generated.h"
+
+DECLARE_LOG_CATEGORY_EXTERN(LogImportGEO, Verbose, All);
+
+
+/** Transient Dataflow container for \c TMap<FString, int32>. */
+USTRUCT(meta = (DataflowFlesh))
+struct FGEOMapStringInt
+{
+	GENERATED_USTRUCT_BODY()
+	TMap<FString, int32> IntVars;
+};
+
+/** Transient Dataflow container for \c TMap<FString, TArray<int32>>. */
+USTRUCT(meta = (DataflowFlesh))
+struct FGEOMapStringArrayInt
+{
+	GENERATED_USTRUCT_BODY()
+	TMap<FString, TArray<int32>> IntVectorVars;
+};
+
+/** Transient Dataflow container for \c TMap<FString, TArray<float>>. */
+USTRUCT(meta = (DataflowFlesh))
+struct FGEOMapStringArrayFloat
+{
+	GENERATED_USTRUCT_BODY()
+	TMap<FString, TArray<float>> FloatVectorVars;
+};
+
+/** Extract a named integer from the results of an ImportGEO node. */
+USTRUCT(meta = (DataflowFlesh))
+struct FExtractGEOInt : public FDataflowNode
+{
+	GENERATED_USTRUCT_BODY()
+	DATAFLOW_NODE_DEFINE_INTERNAL(FExtractGEOInt, "ExtractGEOInt", "Flesh", "")
+
+public:
+	/** IntVars output from ImportGEO node. */
+	UPROPERTY(EditAnywhere, Category = "Dataflow", meta = (DataflowInput, DisplayName = "IntVars"))
+	FGEOMapStringInt IntVars;
+
+	/** Variable name to extract the value for. (See output log from ImportGEO for available list.) */
+	UPROPERTY(EditAnywhere, Category = "Dataflow")
+	FString VarName;
+
+	/** Output value. */
+	UPROPERTY(EditAnywhere, Category = "Dataflow", meta = (DataflowOutput))
+	int32 Value = 0;
+
+	FExtractGEOInt(const Dataflow::FNodeParameters& InParam, FGuid InGuid = FGuid::NewGuid())
+		: FDataflowNode(InParam, InGuid)
+	{
+		RegisterInputConnection(&IntVars);
+		RegisterOutputConnection(&Value);
+	}
+
+	virtual void Evaluate(Dataflow::FContext& Context, const FDataflowOutput* Out) const override;
+};
+
+/** Extract a named integer array from the results of an ImportGEO node. */
+USTRUCT(meta = (DataflowFlesh))
+struct FExtractGEOIntVector : public FDataflowNode
+{
+	GENERATED_USTRUCT_BODY()
+	DATAFLOW_NODE_DEFINE_INTERNAL(FExtractGEOIntVector, "ExtractGEOIntVector", "Flesh", "")
+
+public:
+	/** IntVectorVars output from ImportGEO node. */
+	UPROPERTY(EditAnywhere, Category = "Dataflow", meta = (DataflowInput, DisplayName = "IntVars"))
+	FGEOMapStringArrayInt IntVectorVars;
+
+	/** Variable name to extract the value for. (See output log from ImportGEO for available list.) */
+	UPROPERTY(EditAnywhere, Category = "Dataflow")
+	FString VarName;
+
+	/** Output value. */
+	UPROPERTY(EditAnywhere, Category = "Dataflow", meta = (DataflowOutput))
+	TArray<int32> Value;
+
+	FExtractGEOIntVector(const Dataflow::FNodeParameters& InParam, FGuid InGuid = FGuid::NewGuid())
+		: FDataflowNode(InParam, InGuid)
+	{
+		RegisterInputConnection(&IntVectorVars);
+		RegisterOutputConnection(&Value);
+	}
+
+	virtual void Evaluate(Dataflow::FContext& Context, const FDataflowOutput* Out) const override;
+};
+
+/** Extract a named float array from the results of an ImportGEO node. */
+USTRUCT(meta = (DataflowFlesh))
+struct FExtractGEOFloatVector : public FDataflowNode
+{
+	GENERATED_USTRUCT_BODY()
+	DATAFLOW_NODE_DEFINE_INTERNAL(FExtractGEOFloatVector, "ExtractGEOFloatVector", "Flesh", "")
+
+public:
+	/** FloatVectorVars output from ImportGEO node. */
+	UPROPERTY(EditAnywhere, Category = "Dataflow", meta = (DataflowInput, DisplayName = "IntVars"))
+	FGEOMapStringArrayFloat FloatVectorVars;
+
+	/** Variable name to extract the value for. (See output log from ImportGEO for available list.) */
+	UPROPERTY(EditAnywhere, Category = "Dataflow")
+	FString VarName;
+
+	/** Output value. */
+	UPROPERTY(EditAnywhere, Category = "Dataflow", meta = (DataflowOutput))
+	TArray<float> Value;
+
+	FExtractGEOFloatVector(const Dataflow::FNodeParameters& InParam, FGuid InGuid = FGuid::NewGuid())
+		: FDataflowNode(InParam, InGuid)
+	{
+		RegisterInputConnection(&FloatVectorVars);
+		RegisterOutputConnection(&Value);
+	}
+
+	virtual void Evaluate(Dataflow::FContext& Context, const FDataflowOutput* Out) const override;
+};
+
+/** Import data from GEO file. */
+USTRUCT(meta = (DataflowFlesh))
+struct FImportGEO : public FDataflowNode
+{
+	GENERATED_USTRUCT_BODY()
+	DATAFLOW_NODE_DEFINE_INTERNAL(FImportGEO, "ImportGEO", "Flesh", "")
+public:
+    
+	/** GEO filename. */
+	UPROPERTY(EditAnywhere, Category = "Dataflow", meta = (FilePathFilter = "geo"))
+	FFilePath Filename;
+
+	/** Passthrough geometry collection, used only for tetrahedron mesh extraction. */
+	UPROPERTY(meta = (DataflowInput, DataflowOutput, DataflowPassthrough = "Collection", DisplayName = "Collection"))
+	FManagedArrayCollection Collection;
+
+	//
+	// TetMesh 
+	//
+
+	/** Import the tetrahedron mesh from the GEO file. */
+	UPROPERTY(EditAnywhere, Category = "Dataflow")
+	bool bImportTetrahedronMesh = true;
+
+	/** 
+	 * If \c false, all triangles on the interior of the tetrahedral mesh are retained 
+	 * in the mesh, which is useful for debugging, but makes rendering more costly. 
+	 */
+	UPROPERTY(EditAnywhere, Category = "Dataflow")
+	bool bDiscardInteriorTriangles = true;
+
+	//
+	// GEO Variables
+	//
+
+	/** Integer outputs. Use ExtractGEOInt node. */
+	UPROPERTY(EditAnywhere, Category = "Dataflow", meta = (DataflowOutput, DisplayName = "IntVars"))
+	FGEOMapStringInt IntVarsOutput;
+
+	/** Integer array outputs. Use ExtractGEOIntVector node. */
+	UPROPERTY(EditAnywhere, Category = "Dataflow", meta = (DataflowOutput, DisplayName = "IntVectorVars"))
+	FGEOMapStringArrayInt IntVectorVarsOutput;
+
+	/** Float array outputs. Use ExtractGEOFloatVector node. */
+	UPROPERTY(EditAnywhere, Category = "Dataflow", meta = (DataflowOutput, DisplayName = "FloatVectorVars"))
+	FGEOMapStringArrayFloat FloatVectorVarsOutput;
+
+	FImportGEO(const Dataflow::FNodeParameters& InParam, FGuid InGuid = FGuid::NewGuid())
+		: FDataflowNode(InParam, InGuid)
+	{
+		RegisterInputConnection(&Collection);
+		RegisterOutputConnection(&Collection, &Collection);
+		RegisterOutputConnection(&IntVarsOutput);
+		RegisterOutputConnection(&IntVectorVarsOutput);
+		RegisterOutputConnection(&FloatVectorVarsOutput);
+	}
+
+	virtual void Evaluate(Dataflow::FContext& Context, const FDataflowOutput* Out) const override;
+
+protected:
+	bool ReadGEOFile(const bool PrintStats=true) const;
+
+	mutable FString FilePath;
+	mutable FDateTime FileModTime;
+	mutable TMap<FString, int32> IntVars;
+	mutable TMap<FString, TArray<int32>> IntVectorVars;
+	mutable TMap<FString, TArray<float>> FloatVectorVars;
+	mutable TMap<FString, TPair<TArray<std::string>, TArray<int32>>> IndexedStringVars;
+};
+
+namespace Dataflow
+{
+	void RegisterChaosFleshImportGEONodes();
+}
diff --git a/Engine/Plugins/Experimental/ChaosFlesh/Source/ChaosFleshNodes/Private/Dataflow/ChaosFleshAddKinematicParticlesNode.cpp b/Engine/Plugins/Experimental/ChaosFlesh/Source/ChaosFleshNodes/Private/Dataflow/ChaosFleshAddKinematicParticlesNode.cpp
new file mode 100644
--- /dev/null
+++ b/Engine/Plugins/Experimental/ChaosFlesh/Source/ChaosFleshNodes/Private/Dataflow/ChaosFleshAddKinematicParticlesNode.cpp
@@ -0,0 +1,127 @@
+// Copyright Epic Games, Inc. All Rights Reserved.
+
+#include "Dataflow/ChaosFleshAddKinematicParticlesNode.h"
+
+#include "Dataflow/DataflowEngineUtil.h"
+#include "Engine/SkeletalMesh.h"
+#include "GeometryCollection/Facades/CollectionKinematicBindingFacade.h"
+#include "GeometryCollection/Facades/CollectionVertexBoneWeightsFacade.h"
+#include "GeometryCollection/GeometryCollection.h"
+#include "BoneWeights.h"
+
+#include UE_INLINE_GENERATED_CPP_BY_NAME(ChaosFleshAddKinematicParticlesNode)
+
+//DEFINE_LOG_CATEGORY_STATIC(ChaosFleshAddKinematicParticlesDataflowNodeLog, Log, All);
+
+void FAddKinematicParticlesDataflowNode::Evaluate(Dataflow::FContext& Context, const FDataflowOutput* Out) const
+{
+	if (Out->IsA<DataType>(&Collection) || Out->IsA<TArray<int32>>(&TargetIndicesOut))
+	{
+		DataType InCollection = GetValue<DataType>(Context, &Collection);
+		TArray<int32> TargetIndices;
+		if (TManagedArray<FVector3f>* Vertices = InCollection.FindAttribute<FVector3f>("Vertex", FGeometryCollection::VerticesGroup))
+		{
+			if (FindInput(&VertexIndicesIn) && FindInput(&VertexIndicesIn)->GetConnection())
+			{
+				TArray<int32> BoundVerts;
+				TArray<float> BoundWeights;
+
+				for (int32 SelectionIndex : GetValue<TArray<int32>>(Context, &VertexIndicesIn))
+				{
+					if (0 <= SelectionIndex && SelectionIndex < Vertices->Num())
+					{
+						BoundVerts.Add(SelectionIndex);
+					}
+				}
+				if (BoundVerts.Num())
+				{
+					BoundWeights.Init(1.0, BoundVerts.Num());
+					GeometryCollection::Facades::FKinematicBindingFacade Kinematics(InCollection);
+					Kinematics.AddKinematicBinding(Kinematics.SetBoneBindings(INDEX_NONE, BoundVerts, BoundWeights));
+				}
+			}
+			else if (TObjectPtr<USkeletalMesh> SkeletalMesh = GetValue<TObjectPtr<USkeletalMesh>>(Context, &SkeletalMeshIn))
+			{
+				int32 IndexValue = GetValue<int32>(Context, &BoneIndexIn);
+				if (IndexValue != INDEX_NONE)
+				{
+					TArray<FTransform> ComponentPose;
+					Dataflow::Animation::GlobalTransforms(SkeletalMesh->GetRefSkeleton(), ComponentPose);
+
+					TArray<int32> BranchIndices;
+					if (SkeletalSelectionMode == ESkeletalSeletionMode::Dataflow_SkeletalSelection_Branch)
+					{
+						TArray<int32> ToProcess;
+						int32 CurrentIndex = IndexValue;
+						while (SkeletalMesh->GetRefSkeleton().IsValidIndex(CurrentIndex))
+						{
+							TArray<int32> Buffer;
+							SkeletalMesh->GetRefSkeleton().GetDirectChildBones(CurrentIndex, Buffer);
+
+							if (Buffer.Num())
+							{
+								ToProcess.Append(Buffer);
+							}
+
+							BranchIndices.Add(CurrentIndex);
+
+							CurrentIndex = INDEX_NONE;
+							if (ToProcess.Num())
+							{
+								CurrentIndex = ToProcess.Pop();
+							}
+						}
+					}
+					else // ESkeletalSeletionMode::Dataflow_SkeletalSelection_Single
+					{
+						BranchIndices.Add(IndexValue);
+					}
+
+					// Add standalone particles, not bound to a transform group - so for these particles BoneMap = INDEX_NONE.
+					int32 ParticleIndex = InCollection.AddElements(BranchIndices.Num(), FGeometryCollection::VerticesGroup);
+					TManagedArray<FVector3f>& CurrentVertices = InCollection.ModifyAttribute<FVector3f>("Vertex", FGeometryCollection::VerticesGroup);
+
+					for (int32 Index = 0; Index < BranchIndices.Num(); Index++)
+					{
+						const int32 BoneIndex = BranchIndices[Index];
+						FVector3f BonePosition(ComponentPose[BoneIndex].GetTranslation());
+
+						CurrentVertices[ParticleIndex+Index] = BonePosition;
+
+						TArray<int32> BoundVerts;
+						TArray<float> BoundWeights;
+
+						BoundVerts.Add(ParticleIndex + Index);
+						BoundWeights.Add(1.0);
+						TargetIndices.Emplace(ParticleIndex + Index);
+
+						if (BoundVerts.Num())
+						{
+							GeometryCollection::Facades::FKinematicBindingFacade Kinematics(InCollection);
+							Kinematics.AddKinematicBinding(Kinematics.SetBoneBindings(BoneIndex, BoundVerts, BoundWeights));
+						}
+					}
+
+					//debugging code for tet binding:
+					//int32 ParticleIndex1 = InCollection.AddElements(1, FGeometryCollection::VerticesGroup);
+					//TManagedArray<FVector3f>& CurrentVertices1 = InCollection.ModifyAttribute<FVector3f>("Vertex", FGeometryCollection::VerticesGroup);
+					//CurrentVertices1[ParticleIndex1][0] = -40.f;
+					//CurrentVertices1[ParticleIndex1][1] = -50.f;
+					//CurrentVertices1[ParticleIndex1][2] = -10.f;
+					//GeometryCollection::Facades::FKinematicBindingFacade Kinematics(InCollection);
+					//TArray<int32> BoundVerts;
+					//TArray<float> BoundWeights;
+					//BoundVerts.Add(ParticleIndex1);
+					//BoundWeights.Add(1.0);
+					//Kinematics.AddKinematicBinding(Kinematics.SetBoneBindings(BranchIndices[0], BoundVerts, BoundWeights));
+					//TargetIndices.Emplace(ParticleIndex1);
+
+					GeometryCollection::Facades::FVertexBoneWeightsFacade(InCollection).AddBoneWeightsFromKinematicBindings();
+				}
+			}
+		}
+		SetValue(Context, MoveTemp(InCollection), &Collection);
+		SetValue(Context, MoveTemp(TargetIndices), &TargetIndicesOut);
+	}
+}
+
diff --git a/Engine/Plugins/Experimental/ChaosFlesh/Source/ChaosFleshNodes/Private/Dataflow/ChaosFleshBindForRenderToSkeletalMeshNode.cpp b/Engine/Plugins/Experimental/ChaosFlesh/Source/ChaosFleshNodes/Private/Dataflow/ChaosFleshBindForRenderToSkeletalMeshNode.cpp
deleted file mode 100644
diff --git a/Engine/Plugins/Experimental/ChaosFlesh/Source/ChaosFleshNodes/Private/Dataflow/ChaosFleshComputeFiberFieldNode.cpp b/Engine/Plugins/Experimental/ChaosFlesh/Source/ChaosFleshNodes/Private/Dataflow/ChaosFleshComputeFiberFieldNode.cpp
new file mode 100644
--- /dev/null
+++ b/Engine/Plugins/Experimental/ChaosFlesh/Source/ChaosFleshNodes/Private/Dataflow/ChaosFleshComputeFiberFieldNode.cpp
@@ -0,0 +1,188 @@
+// Copyright Epic Games, Inc. All Rights Reserved.
+
+#include "Dataflow/ChaosFleshComputeFiberFieldNode.h"
+
+#include "Chaos/Math/Poisson.h"
+#include "ChaosFlesh/ChaosFlesh.h"
+#include "ChaosFlesh/TetrahedralCollection.h"
+
+#include UE_INLINE_GENERATED_CPP_BY_NAME(ChaosFleshComputeFiberFieldNode)
+
+void FComputeFiberFieldNode::Evaluate(Dataflow::FContext& Context, const FDataflowOutput* Out) const
+{
+	if (Out->IsA<FManagedArrayCollection>(&Collection))
+	{
+		//
+		// Gather inputs
+		//
+
+		FManagedArrayCollection InCollection = GetValue<FManagedArrayCollection>(Context, &Collection);
+		TArray<int32> InOriginIndices = GetValue<TArray<int32>>(Context, &OriginIndices);
+		TArray<int32> InInsertionIndices = GetValue<TArray<int32>>(Context, &InsertionIndices);
+
+		// Tetrahedra
+		TManagedArray<FIntVector4>* Elements = InCollection.FindAttribute<FIntVector4>(
+			FTetrahedralCollection::TetrahedronAttribute, FTetrahedralCollection::TetrahedralGroup);
+		if (!Elements)
+		{
+			UE_LOG(LogChaosFlesh, Warning,
+				TEXT("ComputeFiberFieldNode: Failed to find geometry collection attr '%s' in group '%s'"),
+				*FTetrahedralCollection::TetrahedronAttribute.ToString(), *FTetrahedralCollection::TetrahedralGroup.ToString());
+			Out->SetValue(MoveTemp(InCollection), Context);
+			return;
+		}
+
+		// Vertices
+		TManagedArray<FVector3f>* Vertex = InCollection.FindAttribute<FVector3f>("Vertex", "Vertices");
+		if (!Vertex)
+		{
+			UE_LOG(LogChaosFlesh, Warning,
+				TEXT("ComputeFiberFieldNode: Failed to find geometry collection attr 'Vertex' in group 'Vertices'"));
+			Out->SetValue(MoveTemp(InCollection), Context);
+			return;
+		}
+
+		// Incident elements
+		TManagedArray<TArray<int32>>* IncidentElements = InCollection.FindAttribute<TArray<int32>>(
+			FTetrahedralCollection::IncidentElementsAttribute, FGeometryCollection::VerticesGroup);
+		if (!IncidentElements)
+		{
+			UE_LOG(LogChaosFlesh, Warning,
+				TEXT("ComputeFiberFieldNode: Failed to find geometry collection attr '%s' in group '%s'"),
+				*FTetrahedralCollection::IncidentElementsAttribute.ToString(), *FGeometryCollection::VerticesGroup.ToString());
+			Out->SetValue(MoveTemp(InCollection), Context);
+			return;
+		}
+		TManagedArray<TArray<int32>>* IncidentElementsLocalIndex = InCollection.FindAttribute<TArray<int32>>(
+			FTetrahedralCollection::IncidentElementsLocalIndexAttribute, FGeometryCollection::VerticesGroup);
+		if (!IncidentElementsLocalIndex)
+		{
+			UE_LOG(LogChaosFlesh, Warning,
+				TEXT("ComputeFiberFieldNode: Failed to find geometry collection attr '%s' in group '%s'"),
+				*FTetrahedralCollection::IncidentElementsLocalIndexAttribute.ToString(), *FGeometryCollection::VerticesGroup.ToString());
+			Out->SetValue(MoveTemp(InCollection), Context);
+			return;
+		}
+
+		//
+		// Pull Origin & Insertion data out of the geometry collection.  We may want other ways of specifying
+		// these via an input on the node...
+		//
+
+		// Origin & Insertion
+		TManagedArray<int32>* Origin = nullptr; 
+		TManagedArray<int32>* Insertion = nullptr;
+		if (InOriginIndices.IsEmpty() || InInsertionIndices.IsEmpty())
+		{
+			// Origin & Insertion group
+			if (OriginInsertionGroupName.IsEmpty())
+			{
+				UE_LOG(LogChaosFlesh, Warning, TEXT("ComputeFiberFieldNode: Attr 'OriginInsertionGroupName' cannot be empty."));
+				Out->SetValue(MoveTemp(InCollection), Context);
+				return;
+			}
+
+			// Origin vertices
+			if (InOriginIndices.IsEmpty())
+			{
+				if (OriginVertexFieldName.IsEmpty())
+				{
+					UE_LOG(LogChaosFlesh, Warning, TEXT("ComputeFiberFieldNode: Attr 'OriginVertexFieldName' cannot be empty."));
+					Out->SetValue(MoveTemp(InCollection), Context);
+					return;
+				}
+				Origin = InCollection.FindAttribute<int32>(FName(OriginVertexFieldName), FName(OriginInsertionGroupName));
+				if (!Origin)
+				{
+					UE_LOG(LogChaosFlesh, Warning,
+						TEXT("ComputeFiberFieldNode: Failed to find geometry collection attr '%s' in group '%s'"),
+						*OriginVertexFieldName, *OriginInsertionGroupName);
+					Out->SetValue(MoveTemp(InCollection), Context);
+					return;
+				}
+			}
+
+			// Insertion vertices
+			if (InInsertionIndices.IsEmpty())
+			{
+				if (InsertionVertexFieldName.IsEmpty())
+				{
+					UE_LOG(LogChaosFlesh, Warning, TEXT("ComputeFiberFieldNode: Attr 'InsertionVertexFieldName' cannot be empty."));
+					Out->SetValue(MoveTemp(InCollection), Context);
+					return;
+				}
+				Insertion = InCollection.FindAttribute<int32>(FName(InsertionVertexFieldName), FName(OriginInsertionGroupName));
+				if (!Insertion)
+				{
+					UE_LOG(LogChaosFlesh, Warning,
+						TEXT("ComputeFiberFieldNode: Failed to find geometry collection attr '%s' in group '%s'"),
+						*InsertionVertexFieldName, *OriginInsertionGroupName);
+					Out->SetValue(MoveTemp(InCollection), Context);
+					return;
+				}
+			}
+		}
+
+		//
+		// Do the thing
+		//
+
+		TArray<FVector3f> FiberDirs =
+			ComputeFiberField(*Elements, *Vertex, *IncidentElements, *IncidentElementsLocalIndex, 
+				Origin ? Origin->GetConstArray() : InOriginIndices,
+				Insertion ? Insertion->GetConstArray() : InInsertionIndices);
+
+		//
+		// Set output(s)
+		//
+
+		TManagedArray<FVector3f>* FiberDirections =
+			InCollection.FindAttribute<FVector3f>("FiberDirection", FTetrahedralCollection::TetrahedralGroup);
+		if (!FiberDirections)
+		{
+			FiberDirections =
+				&InCollection.AddAttribute<FVector3f>("FiberDirection", FTetrahedralCollection::TetrahedralGroup);
+		}
+		(*FiberDirections) = MoveTemp(FiberDirs);
+
+		Out->SetValue(MoveTemp(InCollection), Context);
+	}
+}
+
+TArray<int32>
+FComputeFiberFieldNode::GetNonZeroIndices(const TArray<uint8>& Map) const
+{
+	int32 NumNonZero = 0;
+	for (int32 i = 0; i < Map.Num(); i++)
+		if (Map[i])
+			NumNonZero++;
+	TArray<int32> Indices; Indices.AddUninitialized(NumNonZero);
+	int32 Idx = 0;
+	for (int32 i = 0; i < Map.Num(); i++)
+		if (Map[i])
+			Indices[Idx++] = i;
+	return Indices;
+}
+
+TArray<FVector3f>
+FComputeFiberFieldNode::ComputeFiberField(
+	const TManagedArray<FIntVector4>& Elements,
+	const TManagedArray<FVector3f>& Vertex,
+	const TManagedArray<TArray<int32>>& IncidentElements,
+	const TManagedArray<TArray<int32>>& IncidentElementsLocalIndex,
+	const TArray<int32>& Origin,
+	const TArray<int32>& Insertion) const
+{
+	TArray<FVector3f> Directions;
+	Chaos::ComputeFiberField<float>(
+		Elements.GetConstArray(),
+		Vertex.GetConstArray(),
+		IncidentElements.GetConstArray(),
+		IncidentElementsLocalIndex.GetConstArray(),
+		Origin,
+		Insertion,
+		Directions,
+		MaxIterations,
+		Tolerance);
+	return Directions;
+}
diff --git a/Engine/Plugins/Experimental/ChaosFlesh/Source/ChaosFleshNodes/Private/Dataflow/ChaosFleshComputeIslandsNode.cpp b/Engine/Plugins/Experimental/ChaosFlesh/Source/ChaosFleshNodes/Private/Dataflow/ChaosFleshComputeIslandsNode.cpp
new file mode 100644
--- /dev/null
+++ b/Engine/Plugins/Experimental/ChaosFlesh/Source/ChaosFleshNodes/Private/Dataflow/ChaosFleshComputeIslandsNode.cpp
@@ -0,0 +1,53 @@
+// Copyright Epic Games, Inc. All Rights Reserved.
+
+#include "Dataflow/ChaosFleshComputeIslandsNode.h"
+
+#include "ChaosFlesh/TetrahedralCollection.h"
+
+#include UE_INLINE_GENERATED_CPP_BY_NAME(ChaosFleshComputeIslandsNode)
+
+void FComputeIslandsNode::Evaluate(Dataflow::FContext& Context, const FDataflowOutput* Out) const
+{
+	if (Out->IsA<FManagedArrayCollection>(&Collection))
+	{
+		FManagedArrayCollection InCollection = GetValue<FManagedArrayCollection>(Context, &Collection);
+
+		TManagedArray<int32>& ParticleComponentIndex = InCollection.AddAttribute<int32>("ComponentIndex", FGeometryCollection::VerticesGroup);
+		TManagedArray<FIntVector4>* Elements = InCollection.FindAttribute<FIntVector4>(
+			FTetrahedralCollection::TetrahedronAttribute, FTetrahedralCollection::TetrahedralGroup);
+
+
+		if (Elements)
+		{
+
+			int32 VertsNum = InCollection.NumElements(FGeometryCollection::VerticesGroup);
+			int32 TetsNum = InCollection.NumElements(FTetrahedralCollection::TetrahedralGroup);
+			if (VertsNum && TetsNum)
+			{
+				TArray<TArray<int32>> ConnectedComponents;
+				Chaos::Utilities::FindConnectedRegions(Elements->GetConstArray(), ConnectedComponents);
+				TManagedArray<int32>& ComponentIndex = InCollection.ModifyAttribute<int32>("ComponentIndex", FGeometryCollection::VerticesGroup);
+
+				ComponentIndex.Fill(INDEX_NONE); //Isolated points will get index -1 
+
+				for (int32 i = 0; i < ConnectedComponents.Num(); i++)
+				{
+					for (int32 j = 0; j < ConnectedComponents[i].Num(); j++)
+					{
+						int32 ElementIndex = ConnectedComponents[i][j];
+						for (int32 ie = 0; ie < 4; ie++)
+						{
+							int32 ParticleIndex = (*Elements)[ElementIndex][ie];
+							if (ComponentIndex[ParticleIndex] == INDEX_NONE)
+							{
+								ComponentIndex[ParticleIndex] = i;
+							}
+						}
+					}
+				}
+			}
+		}
+		
+		SetValue(Context, MoveTemp(InCollection), &Collection);
+	}
+}
diff --git a/Engine/Plugins/Experimental/ChaosFlesh/Source/ChaosFleshNodes/Private/Dataflow/ChaosFleshComputeMuscleActivationNode.cpp b/Engine/Plugins/Experimental/ChaosFlesh/Source/ChaosFleshNodes/Private/Dataflow/ChaosFleshComputeMuscleActivationNode.cpp
new file mode 100644
--- /dev/null
+++ b/Engine/Plugins/Experimental/ChaosFlesh/Source/ChaosFleshNodes/Private/Dataflow/ChaosFleshComputeMuscleActivationNode.cpp
@@ -0,0 +1,100 @@
+// Copyright Epic Games, Inc. All Rights Reserved.
+
+#include "Dataflow/ChaosFleshComputeMuscleActivationNode.h"
+
+#include "GeometryCollection/Facades/CollectionMeshFacade.h"
+#include "GeometryCollection/Facades/CollectionMuscleActivationFacade.h"
+#include "ChaosFlesh/TetrahedralCollection.h"
+
+#include UE_INLINE_GENERATED_CPP_BY_NAME(ChaosFleshComputeMuscleActivationNode)
+
+void FComputeMuscleActivationDataNode::Evaluate(Dataflow::FContext& Context, const FDataflowOutput* Out) const
+{
+	if (Out->IsA<FManagedArrayCollection>(&Collection))
+	{
+		FManagedArrayCollection InCollection = GetValue<FManagedArrayCollection>(Context, &Collection);
+		TArray<int32> InOriginIndices = GetValue<TArray<int32>>(Context, &OriginIndicesIn);
+		TArray<int32> InInsertionIndices = GetValue<TArray<int32>>(Context, &InsertionIndicesIn);
+		TArray<float> OriginInsertionRestLength;
+		if (InOriginIndices.Num() > 0 && InInsertionIndices.Num() > 0)
+		{
+			if (TManagedArray<FVector3f>* Vertex = InCollection.FindAttribute<FVector3f>("Vertex", "Vertices"))
+			{
+				if (TManagedArray<FIntVector4>* Elements = InCollection.FindAttribute<FIntVector4>(FTetrahedralCollection::TetrahedronAttribute, FTetrahedralCollection::TetrahedralGroup))
+				{
+					if (TManagedArray<FVector3f>* FiberDirections = InCollection.FindAttribute<FVector3f>("FiberDirection", FTetrahedralCollection::TetrahedralGroup))
+					{
+						GeometryCollection::Facades::FMuscleActivationFacade FMuscleActivation(InCollection);
+						TArray<TArray<int32>> MuscleActivationElements;
+						TArray<int32> OriginNodes;
+						TArray<int32> InsertionNodes;
+						GeometryCollection::Facades::FCollectionMeshFacade MeshFacade(InCollection);
+						TArray<int32> ComponentIndex = MeshFacade.GetGeometryGroupIndexArray();
+						TMap<int32, int32> ComponentToIndex; //Component index to muscle index
+						for (int32 i = 0; i < InOriginIndices.Num(); i++)
+						{
+							if (!ComponentToIndex.Contains(ComponentIndex[InOriginIndices[i]]))
+							{
+								//TODO: use one pair of separately defined origin and insertion for each muscle instead of choosing the first from kinematic origins and last from insertions
+								ComponentToIndex.Add(ComponentIndex[InOriginIndices[i]], OriginNodes.Num());
+								OriginNodes.Add(InOriginIndices[i]);
+							}
+						}
+						InsertionNodes.Init(INDEX_NONE, OriginNodes.Num());
+						for (int32 i = 0; i < InInsertionIndices.Num(); i++)
+						{
+							if (!ComponentToIndex.Contains(ComponentIndex[InInsertionIndices[i]]))
+							{
+								ensureMsgf(false, TEXT("No origin in this component"));
+							}
+							else if (InsertionNodes[ComponentToIndex[ComponentIndex[InInsertionIndices[i]]]] == INDEX_NONE)
+							{
+								InsertionNodes[ComponentToIndex[ComponentIndex[InInsertionIndices[i]]]] = InInsertionIndices[i];
+							}
+						}
+						MuscleActivationElements.SetNum(OriginNodes.Num());
+						for (int32 ElemIdx = 0; ElemIdx < Elements->Num(); ElemIdx++)
+						{
+							if (ComponentToIndex.Contains(ComponentIndex[(*Elements)[ElemIdx][0]]))
+							{
+								MuscleActivationElements[ComponentToIndex[ComponentIndex[(*Elements)[ElemIdx][0]]]].Add(ElemIdx);
+							}
+						}
+						OriginInsertionRestLength.SetNum(OriginNodes.Num());
+
+						for (int32 MuscleComponentIdx = 0; MuscleComponentIdx < OriginNodes.Num(); MuscleComponentIdx++)
+						{
+							if (ensureMsgf(InsertionNodes[MuscleComponentIdx] != INDEX_NONE, TEXT("InsertionNodes[%d] is not a pair"), MuscleComponentIdx))
+							{
+								GeometryCollection::Facades::FMuscleActivationData MuscleActivationData;
+								MuscleActivationData.MuscleActivationElement = MuscleActivationElements[MuscleComponentIdx];
+								MuscleActivationData.OriginInsertionPair = FIntVector2(OriginNodes[MuscleComponentIdx], InsertionNodes[MuscleComponentIdx]);
+								MuscleActivationData.OriginInsertionRestLength = ((*Vertex)[OriginNodes[MuscleComponentIdx]] - (*Vertex)[InsertionNodes[MuscleComponentIdx]]).Size();
+								MuscleActivationData.FiberDirectionMatrix.SetNum(MuscleActivationElements[MuscleComponentIdx].Num());
+								for (int32 LocalElemIdx = 0; LocalElemIdx < MuscleActivationElements[MuscleComponentIdx].Num(); LocalElemIdx++)
+								{
+									FVector3f V = (*FiberDirections)[MuscleActivationElements[MuscleComponentIdx][LocalElemIdx]];
+									// QR decomposition on vvT
+									FVector3f W = V;
+									if (V.X < V.Y)
+									{
+										W.X += 1.f;
+									}
+									else
+									{
+										W.Y += 1.f;
+									}
+									FVector3f U = (V ^ W).GetSafeNormal();
+									W = (U ^ V).GetSafeNormal();
+									MuscleActivationData.FiberDirectionMatrix[LocalElemIdx] = Chaos::PMatrix33d(V, W, U);
+								}
+								FMuscleActivation.AddMuscleActivationData(MuscleActivationData);
+							}
+						}
+					}
+				}
+			}
+		}
+		Out->SetValue(MoveTemp(InCollection), Context);
+	}
+}
diff --git a/Engine/Plugins/Experimental/ChaosFlesh/Source/ChaosFleshNodes/Private/Dataflow/ChaosFleshCreateTetrahedronNode.cpp b/Engine/Plugins/Experimental/ChaosFlesh/Source/ChaosFleshNodes/Private/Dataflow/ChaosFleshCreateTetrahedronNode.cpp
--- a/Engine/Plugins/Experimental/ChaosFlesh/Source/ChaosFleshNodes/Private/Dataflow/ChaosFleshCreateTetrahedronNode.cpp
+++ b/Engine/Plugins/Experimental/ChaosFlesh/Source/ChaosFleshNodes/Private/Dataflow/ChaosFleshCreateTetrahedronNode.cpp
@@ -1,7 +1,6 @@
 // Copyright Epic Games, Inc. All Rights Reserved.
 
 #include "Dataflow/ChaosFleshCreateTetrahedronNode.h"
-#include "Dataflow/ChaosFleshEngineAssetNodes.h"
 
 #include "Async/ParallelFor.h"
 #include "Chaos/Deformable/Utilities.h"
diff --git a/Engine/Plugins/Experimental/ChaosFlesh/Source/ChaosFleshNodes/Private/Dataflow/ChaosFleshEngineAssetNodes.cpp b/Engine/Plugins/Experimental/ChaosFlesh/Source/ChaosFleshNodes/Private/Dataflow/ChaosFleshEngineAssetNodes.cpp
deleted file mode 100644
diff --git a/Engine/Plugins/Experimental/ChaosFlesh/Source/ChaosFleshNodes/Private/Dataflow/ChaosFleshFleshAssetTerminalNode.cpp b/Engine/Plugins/Experimental/ChaosFlesh/Source/ChaosFleshNodes/Private/Dataflow/ChaosFleshFleshAssetTerminalNode.cpp
new file mode 100644
--- /dev/null
+++ b/Engine/Plugins/Experimental/ChaosFlesh/Source/ChaosFleshNodes/Private/Dataflow/ChaosFleshFleshAssetTerminalNode.cpp
@@ -0,0 +1,21 @@
+// Copyright Epic Games, Inc. All Rights Reserved.
+
+#include "Dataflow/ChaosFleshFleshAssetTerminalNode.h"
+
+#include UE_INLINE_GENERATED_CPP_BY_NAME(ChaosFleshFleshAssetTerminalNode)
+
+
+void FFleshAssetTerminalDataflowNode::SetAssetValue(TObjectPtr<UObject> Asset, Dataflow::FContext& Context) const
+{
+	if (UFleshAsset* FleshAsset = Cast<UFleshAsset>(Asset.Get()))
+	{
+		const FManagedArrayCollection& InCollection = GetValue<FManagedArrayCollection>(Context, &Collection);
+		FleshAsset->SetCollection(InCollection.NewCopy<FFleshCollection>());
+	}
+}
+
+void FFleshAssetTerminalDataflowNode::Evaluate(Dataflow::FContext& Context, const FDataflowOutput* Out) const
+{
+	const FManagedArrayCollection& InCollection = GetValue<FManagedArrayCollection>(Context, &Collection);
+	SetValue(Context, InCollection, &Collection);
+}
diff --git a/Engine/Plugins/Experimental/ChaosFlesh/Source/ChaosFleshNodes/Private/Dataflow/ChaosFleshGenerateFiberDirectionsNode.cpp b/Engine/Plugins/Experimental/ChaosFlesh/Source/ChaosFleshNodes/Private/Dataflow/ChaosFleshGenerateFiberDirectionsNode.cpp
deleted file mode 100644
diff --git a/Engine/Plugins/Experimental/ChaosFlesh/Source/ChaosFleshNodes/Private/Dataflow/ChaosFleshGenerateOriginInsertionNode.cpp b/Engine/Plugins/Experimental/ChaosFlesh/Source/ChaosFleshNodes/Private/Dataflow/ChaosFleshGenerateOriginInsertionNode.cpp
new file mode 100644
--- /dev/null
+++ b/Engine/Plugins/Experimental/ChaosFlesh/Source/ChaosFleshNodes/Private/Dataflow/ChaosFleshGenerateOriginInsertionNode.cpp
@@ -0,0 +1,122 @@
+// Copyright Epic Games, Inc. All Rights Reserved.
+
+#include "Dataflow/ChaosFleshGenerateOriginInsertionNode.h"
+
+#include "ChaosFlesh/ChaosFlesh.h"
+#include "ChaosFlesh/TetrahedralCollection.h"
+#include "GeometryCollection/Facades/CollectionMeshFacade.h"
+
+#include UE_INLINE_GENERATED_CPP_BY_NAME(ChaosFleshGenerateOriginInsertionNode)
+
+void FGenerateOriginInsertionNode::Evaluate(Dataflow::FContext& Context, const FDataflowOutput* Out) const
+{
+	if (Out->IsA<FManagedArrayCollection>(&Collection))
+	{
+		//
+		// Gather inputs
+		//
+
+		FManagedArrayCollection InCollection = GetValue<FManagedArrayCollection>(Context, &Collection);
+		TArray<int32> InOriginIndices = GetValue<TArray<int32>>(Context, &OriginIndicesIn);
+		TArray<int32> InInsertionIndices = GetValue<TArray<int32>>(Context, &InsertionIndicesIn);
+		TArray<int32> OutOriginIndices;
+		TArray<int32> OutInsertionIndices;
+		// Tetrahedra
+		TManagedArray<FIntVector4>* Elements = InCollection.FindAttribute<FIntVector4>(
+			FTetrahedralCollection::TetrahedronAttribute, FTetrahedralCollection::TetrahedralGroup);
+		if (!Elements)
+		{
+			UE_LOG(LogChaosFlesh, Warning,
+				TEXT("GenerateOriginInsertionNode: Failed to find geometry collection attr '%s' in group '%s'"),
+				*FTetrahedralCollection::TetrahedronAttribute.ToString(), *FTetrahedralCollection::TetrahedralGroup.ToString());
+			Out->SetValue(MoveTemp(InCollection), Context);
+			return;
+		}
+
+		// Vertices
+		TManagedArray<FVector3f>* Vertex = InCollection.FindAttribute<FVector3f>("Vertex", "Vertices");
+		//TArray<FVector3f>* MeshVertex;
+		if (!Vertex)
+		{
+			UE_LOG(LogChaosFlesh, Warning,
+				TEXT("GenerateOriginInsertionNode: Failed to find geometry collection attr 'Vertex' in group 'Vertices'"));
+			Out->SetValue(MoveTemp(InCollection), Context);
+			return;
+		}
+
+		// Incident elements
+		TManagedArray<TArray<int32>>* IncidentElements = InCollection.FindAttribute<TArray<int32>>(
+			FTetrahedralCollection::IncidentElementsAttribute, FGeometryCollection::VerticesGroup);
+		if (!IncidentElements)
+		{
+			UE_LOG(LogChaosFlesh, Warning,
+				TEXT("GenerateOriginInsertionNode: Failed to find geometry collection attr '%s' in group '%s'"),
+				*FTetrahedralCollection::IncidentElementsAttribute.ToString(), *FGeometryCollection::VerticesGroup.ToString());
+			Out->SetValue(MoveTemp(InCollection), Context);
+			return;
+		}
+		TManagedArray<TArray<int32>>* IncidentElementsLocalIndex = InCollection.FindAttribute<TArray<int32>>(
+			FTetrahedralCollection::IncidentElementsLocalIndexAttribute, FGeometryCollection::VerticesGroup);
+		if (!IncidentElementsLocalIndex)
+		{
+			UE_LOG(LogChaosFlesh, Warning,
+				TEXT("GenerateOriginInsertionNode: Failed to find geometry collection attr '%s' in group '%s'"),
+				*FTetrahedralCollection::IncidentElementsLocalIndexAttribute.ToString(), *FGeometryCollection::VerticesGroup.ToString());
+			Out->SetValue(MoveTemp(InCollection), Context);
+			return;
+		}
+
+		//
+		// Pull Origin & Insertion data out of the geometry collection.  We may want other ways of specifying
+		// these via an input on the node...
+		//
+		auto DoubleVert = [](FVector3f V) { return FVector3d(V.X, V.Y, V.Z); };
+		GeometryCollection::Facades::FCollectionMeshFacade MeshFacade(InCollection);
+		TArray<int32> ComponentIndex = MeshFacade.GetGeometryGroupIndexArray();
+		// Origin vertices
+		if (!InOriginIndices.IsEmpty())
+		{
+			for (int32 i = 0; i < InOriginIndices.Num(); ++i)
+			{
+				if (InOriginIndices[i] < Vertex->Num())
+				{
+					for (int32 j = 0; j < Vertex->Num(); ++j)
+					{
+						if (ComponentIndex[InOriginIndices[i]] == ComponentIndex[j] 
+							&& ComponentIndex[InOriginIndices[i]] >= 0
+							&& ComponentIndex[j] >= 0
+							&& ((*Vertex)[InOriginIndices[i]] - (*Vertex)[j]).Size() < Radius)
+						{
+							OutOriginIndices.Add(j);
+						}
+					}
+				}
+			}
+		}
+
+		// Insertion vertices
+		if (!InInsertionIndices.IsEmpty())
+		{
+			for (int32 i = 0; i < InInsertionIndices.Num(); ++i)
+			{
+				if (InInsertionIndices[i] < Vertex->Num())
+				{
+					for (int32 j = 0; j < Vertex->Num(); ++j)
+					{
+						if (ComponentIndex[InInsertionIndices[i]] == ComponentIndex[j]
+							&& ComponentIndex[InInsertionIndices[i]] >= 0
+							&& ComponentIndex[j] >= 0
+							&& ((*Vertex)[InInsertionIndices[i]] - (*Vertex)[j]).Size() < Radius)
+						{
+							OutInsertionIndices.Add(j);
+						}
+					}
+				}
+			}
+		}
+
+		SetValue(Context, MoveTemp(InCollection), &Collection);
+		SetValue(Context, MoveTemp(OutOriginIndices), &OriginIndicesOut);
+		SetValue(Context, MoveTemp(OutInsertionIndices), &InsertionIndicesOut);
+	}
+}
diff --git a/Engine/Plugins/Experimental/ChaosFlesh/Source/ChaosFleshNodes/Private/Dataflow/ChaosFleshGetFleshAssetNode.cpp b/Engine/Plugins/Experimental/ChaosFlesh/Source/ChaosFleshNodes/Private/Dataflow/ChaosFleshGetFleshAssetNode.cpp
new file mode 100644
--- /dev/null
+++ b/Engine/Plugins/Experimental/ChaosFlesh/Source/ChaosFleshNodes/Private/Dataflow/ChaosFleshGetFleshAssetNode.cpp
@@ -0,0 +1,31 @@
+// Copyright Epic Games, Inc. All Rights Reserved.
+
+#include "Dataflow/ChaosFleshGetFleshAssetNode.h"
+
+#include UE_INLINE_GENERATED_CPP_BY_NAME(ChaosFleshGetFleshAssetNode)
+
+void FGetFleshAssetDataflowNode::Evaluate(Dataflow::FContext& Context, const FDataflowOutput* Out) const
+{
+	if (Out->IsA<FManagedArrayCollection>(&Output))
+	{
+		FManagedArrayCollection Collection;
+		SetValue(Context, MoveTemp(Collection), &Output);
+
+		const UFleshAsset* FleshAssetValue = FleshAsset;
+		if (!FleshAssetValue)
+		{
+			if (const Dataflow::FEngineContext* EngineContext = Context.AsType<Dataflow::FEngineContext>())
+			{
+				FleshAssetValue = Cast<UFleshAsset>(EngineContext->Owner);
+			}
+		}
+
+		if (FleshAssetValue)
+		{
+			if (const FFleshCollection* AssetCollection = FleshAssetValue->GetCollection())
+			{
+				SetValue(Context, (const FManagedArrayCollection&)(*AssetCollection), &Output);
+			}
+		}
+	}
+}
diff --git a/Engine/Plugins/Experimental/ChaosFlesh/Source/ChaosFleshNodes/Private/Dataflow/ChaosFleshGetSurfaceIndicesNode.cpp b/Engine/Plugins/Experimental/ChaosFlesh/Source/ChaosFleshNodes/Private/Dataflow/ChaosFleshGetSurfaceIndicesNode.cpp
new file mode 100644
--- /dev/null
+++ b/Engine/Plugins/Experimental/ChaosFlesh/Source/ChaosFleshNodes/Private/Dataflow/ChaosFleshGetSurfaceIndicesNode.cpp
@@ -0,0 +1,50 @@
+// Copyright Epic Games, Inc. All Rights Reserved.
+
+#include "Dataflow/ChaosFleshGetSurfaceIndicesNode.h"
+
+#include UE_INLINE_GENERATED_CPP_BY_NAME(ChaosFleshGetSurfaceIndicesNode)
+
+void FGetSurfaceIndicesNode::Evaluate(Dataflow::FContext& Context, const FDataflowOutput* Out) const
+{
+	FManagedArrayCollection InCollection = GetValue<FManagedArrayCollection>(Context, &Collection);
+	TArray<int32> SurfaceIndicesLocal;
+	if (TManagedArray<FIntVector>* Indices = InCollection.FindAttribute<FIntVector>("Indices", FGeometryCollection::FacesGroup))
+	{
+		if (FindInput(&GeometryGroupGuidsIn) && FindInput(&GeometryGroupGuidsIn)->GetConnection())
+		{
+			TArray<FString> GeometryGroupGuidsLocal = GetValue<TArray<FString>>(Context, &GeometryGroupGuidsIn);
+			TManagedArray<int32>* IndicesStart = InCollection.FindAttribute<int32>("FaceStart", FGeometryCollection::GeometryGroup);
+			TManagedArray<int32>* IndicesCount = InCollection.FindAttribute<int32>("FaceCount", FGeometryCollection::GeometryGroup);
+			if (TManagedArray<FString>* Guids = InCollection.FindAttribute<FString>("Guid", FGeometryCollection::GeometryGroup))
+			{
+				for (int32 Idx = 0; Idx < IndicesStart->Num(); Idx++)
+				{
+					if (GeometryGroupGuidsLocal.Num() && Guids)
+					{
+						if (GeometryGroupGuidsLocal.Contains((*Guids)[Idx]))
+						{
+							for (int32 i = (*IndicesStart)[Idx]; i < (*IndicesStart)[Idx] + (*IndicesCount)[Idx]; i++)
+							{
+								for (int32 j = 0; j < 3; j++)
+								{
+									SurfaceIndicesLocal.AddUnique((*Indices)[i][j]);
+								}
+							}
+						}
+					}
+				}
+			}
+		}
+		else
+		{
+			for (int32 i = 0; i < Indices->Num(); i++)
+			{
+				for (int32 j = 0; j < 3; j++)
+				{
+					SurfaceIndicesLocal.AddUnique((*Indices)[i][j]);
+				}
+			}
+		}
+	}
+	SetValue(Context, MoveTemp(SurfaceIndicesLocal), &SurfaceIndicesOut);
+}
diff --git a/Engine/Plugins/Experimental/ChaosFlesh/Source/ChaosFleshNodes/Private/Dataflow/ChaosFleshImportGEO.cpp b/Engine/Plugins/Experimental/ChaosFlesh/Source/ChaosFleshNodes/Private/Dataflow/ChaosFleshImportGEO.cpp
deleted file mode 100644
diff --git a/Engine/Plugins/Experimental/ChaosFlesh/Source/ChaosFleshNodes/Private/Dataflow/ChaosFleshIsolateComponentNode.cpp b/Engine/Plugins/Experimental/ChaosFlesh/Source/ChaosFleshNodes/Private/Dataflow/ChaosFleshIsolateComponentNode.cpp
new file mode 100644
--- /dev/null
+++ b/Engine/Plugins/Experimental/ChaosFlesh/Source/ChaosFleshNodes/Private/Dataflow/ChaosFleshIsolateComponentNode.cpp
@@ -0,0 +1,45 @@
+// Copyright Epic Games, Inc. All Rights Reserved.
+
+#include "Dataflow/ChaosFleshIsolateComponentNode.h"
+
+#include UE_INLINE_GENERATED_CPP_BY_NAME(ChaosFleshIsolateComponentNode)
+
+void FIsolateComponentNode::Evaluate(Dataflow::FContext& Context, const FDataflowOutput* Out) const
+{
+	if (Out->IsA<FManagedArrayCollection>(&Collection))
+	{	
+		FManagedArrayCollection InCollection = GetValue<FManagedArrayCollection>(Context, &Collection);
+		TArray<int32> DeleteList;
+		TManagedArray<FIntVector>* Indices = InCollection.FindAttribute<FIntVector>("Indices", FGeometryCollection::FacesGroup);
+		TManagedArray<bool>*FaceVisibility = InCollection.FindAttribute<bool>("Visible", FGeometryCollection::FacesGroup);
+		TManagedArray<int32>* FaceStart = InCollection.FindAttribute<int32>("FaceStart", FGeometryCollection::GeometryGroup);
+		TManagedArray<int32>* FaceCount = InCollection.FindAttribute<int32>("FaceCount", FGeometryCollection::GeometryGroup);
+		if (Indices && FaceVisibility && FaceStart && FaceCount)
+		{
+			FaceVisibility->Fill(false);
+			TSet<int32> GeometrySet;
+			TArray<FString> StrArray;
+			TargetGeometryIndex.ParseIntoArray(StrArray, *FString(" "));
+			for (FString GeometryIdx : StrArray)
+			{
+				if (GeometryIdx.Len() && FCString::IsNumeric(*GeometryIdx))
+				{
+					GeometrySet.Add(FCString::Atoi(*GeometryIdx));
+				}
+			}
+			for (TSet<int32>::TConstIterator It = GeometrySet.CreateConstIterator(); It; ++It)
+			{
+				for (int32 FaceIdx = (*FaceStart)[*It]; FaceIdx < (*FaceStart)[*It] + (*FaceCount)[*It]; FaceIdx++)
+				{
+					(*FaceVisibility)[FaceIdx] = true;
+					DeleteList.Add(FaceIdx);
+				}
+			}
+			if (bDeleteHiddenFaces)
+			{
+				InCollection.RemoveElements(FGeometryCollection::FacesGroup, DeleteList);
+			}
+		}
+		SetValue(Context, MoveTemp(InCollection), &Collection);
+	}
+}
diff --git a/Engine/Plugins/Experimental/ChaosFlesh/Source/ChaosFleshNodes/Private/Dataflow/ChaosFleshNodesPlugin.cpp b/Engine/Plugins/Experimental/ChaosFlesh/Source/ChaosFleshNodes/Private/Dataflow/ChaosFleshNodesPlugin.cpp
--- a/Engine/Plugins/Experimental/ChaosFlesh/Source/ChaosFleshNodes/Private/Dataflow/ChaosFleshNodesPlugin.cpp
+++ b/Engine/Plugins/Experimental/ChaosFlesh/Source/ChaosFleshNodes/Private/Dataflow/ChaosFleshNodesPlugin.cpp
@@ -2,60 +2,75 @@
 
 #include "Dataflow/ChaosFleshNodesPlugin.h"
 
-
+#include "Dataflow/ChaosFleshAddKinematicParticlesNode.h"
+#include "Dataflow/ChaosFleshAppendTetrahedralCollectionNode.h"
 #include "Dataflow/ChaosFleshAuthorSceneCollisionCandidatesNode.h"
 #include "Dataflow/ChaosFleshCalculateTetrehedralMetricsNode.h"
 #include "Dataflow/ChaosFleshCollisionBodyConstraintNode.h"
+#include "Dataflow/ChaosFleshComputeFiberFieldNode.h"
+#include "Dataflow/ChaosFleshComputeIslandsNode.h"
+#include "Dataflow/ChaosFleshComputeMuscleActivationNode.h"
 #include "Dataflow/ChaosFleshCreateTetrahedronNode.h"
-#include "Dataflow/ChaosFleshEngineAssetNodes.h"
-#include "Dataflow/ChaosFleshImportGEO.h"
+#include "Dataflow/ChaosFleshFleshAssetTerminalNode.h"
+#include "Dataflow/ChaosFleshGenerateOriginInsertionNode.h"
+#include "Dataflow/ChaosFleshGenerateSkeletalBindingsNode.h"
+#include "Dataflow/ChaosFleshGenerateSurfaceBindingsNode.h"
+#include "Dataflow/ChaosFleshGetFleshAssetNode.h"
+#include "Dataflow/ChaosFleshGetSurfaceIndicesNode.h"
+#include "Dataflow/ChaosFleshIsolateComponentNode.h"
 #include "Dataflow/ChaosFleshKinematicConstraintNode.h"
 #include "Dataflow/ChaosFleshKinematicOriginInsertionInitializationNode.h"
 #include "Dataflow/ChaosFleshKinematicTetrahedralConstraintNode.h"
 #include "Dataflow/ChaosFleshRadialTetrahedronNode.h"
-#include "Dataflow/ChaosFleshPositionTargetInitializationNodes.h"
 #include "Dataflow/ChaosFleshTriangleMeshSimulationPropertiesNode.h"
 #include "Dataflow/ChaosFleshSkeletalMeshConstraintNode.h"
 #include "Dataflow/ChaosFleshSkinSimulationPropertiesNode.h"
+#include "Dataflow/ChaosFleshSetFleshBonePositionTargetBindingNode.h"
 #include "Dataflow/ChaosFleshSetFleshDefaultPropertiesNode.h"
+#include "Dataflow/ChaosFleshSetVertexTetrahedraPositionTargetBindingNode.h"
+#include "Dataflow/ChaosFleshSetVertexTrianglePositionTargetBindingNode.h"
+#include "Dataflow/ChaosFleshSetVertexVertexPositionTargetBindingNode.h"
 #include "Dataflow/ChaosFleshVertexConstraintNode.h"
+#include "Dataflow/ChaosFleshVisualizeFiberFieldNode.h"
 #include "Dataflow/GeometryCollectionAppendCollectionTransformNode.h"
 #include "Modules/ModuleManager.h"
 
-#include "Dataflow/ChaosFleshAppendTetrahedralCollectionNode.h"
-#include "Dataflow/ChaosFleshBindForRenderToSkeletalMeshNode.h"
-#include "Dataflow/ChaosFleshGenerateSkeletalBindingsNode.h"
-#include "Dataflow/ChaosFleshGenerateSurfaceBindingsNode.h"
-#include "Dataflow/ChaosFleshGenerateFiberDirectionsNode.h"
-
-
 #define LOCTEXT_NAMESPACE "ChaosFleshNodes"
 
 
 void IChaosFleshNodesPlugin::StartupModule()
 {
-	Dataflow::RegisterChaosFleshEngineAssetNodes();
-	Dataflow::RegisterChaosFleshPositionTargetInitializationNodes();
-	Dataflow::RegisterChaosFleshImportGEONodes();
+	DATAFLOW_NODE_REGISTER_CREATION_FACTORY(FAddKinematicParticlesDataflowNode);
 	DATAFLOW_NODE_REGISTER_CREATION_FACTORY(FAppendTetrahedralCollectionDataflowNode);
-	DATAFLOW_NODE_REGISTER_CREATION_FACTORY(FBindForRenderToSkeletalMeshDataflowNode); // todo delete
-	DATAFLOW_NODE_REGISTER_CREATION_FACTORY(FGenerateFiberDirectionsDataflowNode); // todo delete
+	DATAFLOW_NODE_REGISTER_CREATION_FACTORY(FAppendToCollectionTransformAttributeDataflowNode);
+	DATAFLOW_NODE_REGISTER_CREATION_FACTORY(FAuthorSceneCollisionCandidates);
+	DATAFLOW_NODE_REGISTER_CREATION_FACTORY(FCalculateTetMetrics);
+	DATAFLOW_NODE_REGISTER_CREATION_FACTORY(FComputeFiberFieldNode);
+	DATAFLOW_NODE_REGISTER_CREATION_FACTORY(FComputeIslandsNode);
+	DATAFLOW_NODE_REGISTER_CREATION_FACTORY(FComputeMuscleActivationDataNode);
+	DATAFLOW_NODE_REGISTER_CREATION_FACTORY(FCreateTetrahedronDataflowNode);
+	DATAFLOW_NODE_REGISTER_CREATION_FACTORY(FFleshAssetTerminalDataflowNode);
+	DATAFLOW_NODE_REGISTER_CREATION_FACTORY(FGenerateOriginInsertionNode);
 	DATAFLOW_NODE_REGISTER_CREATION_FACTORY(FGenerateSkeletalBindings);
 	DATAFLOW_NODE_REGISTER_CREATION_FACTORY(FGenerateSurfaceBindings);
-	DATAFLOW_NODE_REGISTER_CREATION_FACTORY(FCreateTetrahedronDataflowNode);
-	DATAFLOW_NODE_REGISTER_CREATION_FACTORY(FTriangleMeshSimulationPropertiesDataflowNodes);
-	DATAFLOW_NODE_REGISTER_CREATION_FACTORY(FSkinSimulationPropertiesDataflowNodes);
-	DATAFLOW_NODE_REGISTER_CREATION_FACTORY(FSetFleshDefaultPropertiesNode);
-	DATAFLOW_NODE_REGISTER_CREATION_FACTORY(FKinematicSkeletalMeshInitializationDataflowNode);
+	DATAFLOW_NODE_REGISTER_CREATION_FACTORY(FGetFleshAssetDataflowNode);
+	DATAFLOW_NODE_REGISTER_CREATION_FACTORY(FGetSurfaceIndicesNode);
+	DATAFLOW_NODE_REGISTER_CREATION_FACTORY(FIsolateComponentNode);
 	DATAFLOW_NODE_REGISTER_CREATION_FACTORY(FKinematicBodySetupInitializationDataflowNode);
 	DATAFLOW_NODE_REGISTER_CREATION_FACTORY(FKinematicInitializationDataflowNode);
 	DATAFLOW_NODE_REGISTER_CREATION_FACTORY(FKinematicOriginInsertionInitializationDataflowNode);
+	DATAFLOW_NODE_REGISTER_CREATION_FACTORY(FKinematicSkeletalMeshInitializationDataflowNode);
 	DATAFLOW_NODE_REGISTER_CREATION_FACTORY(FKinematicTetrahedralBindingsDataflowNode);
-	DATAFLOW_NODE_REGISTER_CREATION_FACTORY(FSetVerticesKinematicDataflowNode);
-	DATAFLOW_NODE_REGISTER_CREATION_FACTORY(FAuthorSceneCollisionCandidates);
-	DATAFLOW_NODE_REGISTER_CREATION_FACTORY(FAppendToCollectionTransformAttributeDataflowNode); 
-	DATAFLOW_NODE_REGISTER_CREATION_FACTORY(FCalculateTetMetrics);
 	DATAFLOW_NODE_REGISTER_CREATION_FACTORY(FRadialTetrahedronDataflowNodes);
+	DATAFLOW_NODE_REGISTER_CREATION_FACTORY(FSetFleshBonePositionTargetBindingDataflowNode);
+	DATAFLOW_NODE_REGISTER_CREATION_FACTORY(FSetFleshDefaultPropertiesNode);
+	DATAFLOW_NODE_REGISTER_CREATION_FACTORY(FSetVerticesKinematicDataflowNode);
+	DATAFLOW_NODE_REGISTER_CREATION_FACTORY(FSetVertexTetrahedraPositionTargetBindingDataflowNode);
+	DATAFLOW_NODE_REGISTER_CREATION_FACTORY(FSetVertexTrianglePositionTargetBindingDataflowNode);
+	DATAFLOW_NODE_REGISTER_CREATION_FACTORY(FSetVertexVertexPositionTargetBindingDataflowNode);
+	DATAFLOW_NODE_REGISTER_CREATION_FACTORY(FSkinSimulationPropertiesDataflowNodes);
+	DATAFLOW_NODE_REGISTER_CREATION_FACTORY(FTriangleMeshSimulationPropertiesDataflowNodes);
+	DATAFLOW_NODE_REGISTER_CREATION_FACTORY(FVisualizeFiberFieldNode);
 }
 
 void IChaosFleshNodesPlugin::ShutdownModule()
diff --git a/Engine/Plugins/Experimental/ChaosFlesh/Source/ChaosFleshNodes/Private/Dataflow/ChaosFleshPositionTargetInitializationNodes.cpp b/Engine/Plugins/Experimental/ChaosFlesh/Source/ChaosFleshNodes/Private/Dataflow/ChaosFleshPositionTargetInitializationNodes.cpp
deleted file mode 100644
diff --git a/Engine/Plugins/Experimental/ChaosFlesh/Source/ChaosFleshNodes/Private/Dataflow/ChaosFleshSetFleshBonePositionTargetBindingNode.cpp b/Engine/Plugins/Experimental/ChaosFlesh/Source/ChaosFleshNodes/Private/Dataflow/ChaosFleshSetFleshBonePositionTargetBindingNode.cpp
new file mode 100644
--- /dev/null
+++ b/Engine/Plugins/Experimental/ChaosFlesh/Source/ChaosFleshNodes/Private/Dataflow/ChaosFleshSetFleshBonePositionTargetBindingNode.cpp
@@ -0,0 +1,347 @@
+// Copyright Epic Games, Inc. All Rights Reserved.
+
+#include "Dataflow/ChaosFleshSetFleshBonePositionTargetBindingNode.h"
+
+#include "Engine/SkeletalMesh.h"
+#include "Chaos/BoundingVolumeHierarchy.h"
+#include "GeometryCollection/Facades/CollectionKinematicBindingFacade.h"
+#include "GeometryCollection/Facades/CollectionPositionTargetFacade.h"
+#include "GeometryCollection/Facades/CollectionVertexBoneWeightsFacade.h"
+#include "Rendering/SkeletalMeshModel.h"
+#include "Rendering/SkeletalMeshRenderData.h"
+#include "BoneWeights.h"
+
+#include UE_INLINE_GENERATED_CPP_BY_NAME(ChaosFleshSetFleshBonePositionTargetBindingNode)
+
+//DEFINE_LOG_CATEGORY_STATIC(ChaosFleshSetFleshBonePositionTargetBindingNodeLog, Log, All);
+
+void FSetFleshBonePositionTargetBindingDataflowNode::Evaluate(Dataflow::FContext& Context, const FDataflowOutput* Out) const
+{
+	if (Out->IsA<DataType>(&Collection))
+	{
+		DataType InCollection = GetValue<DataType>(Context, &Collection);
+
+		if (TObjectPtr<USkeletalMesh> BoneSkeletalMesh = GetValue<TObjectPtr<USkeletalMesh>>(Context, &SkeletalMeshIn))
+		{
+			if (TManagedArray<FIntVector>* Indices = InCollection.FindAttribute<FIntVector>("Indices", FGeometryCollection::FacesGroup))
+			{
+				if (TManagedArray<FVector3f>* Vertices = InCollection.FindAttribute<FVector3f>("Vertex", FGeometryCollection::VerticesGroup))
+				{
+
+					FSkeletalMeshRenderData* RenderData = BoneSkeletalMesh->GetResourceForRendering();
+					if (RenderData->LODRenderData.Num())
+					{
+						FSkeletalMeshLODRenderData* LODRenderData = &RenderData->LODRenderData[0];
+						const FPositionVertexBuffer& PositionVertexBuffer =
+							LODRenderData->StaticVertexBuffers.PositionVertexBuffer;
+
+						const FSkinWeightVertexBuffer* SkinWeightVertexBuffer = LODRenderData->GetSkinWeightVertexBuffer();
+						const int32 MaxBoneInfluences = SkinWeightVertexBuffer->GetMaxBoneInfluences();
+						TArray<FTransform> ComponentPose;
+						Dataflow::Animation::GlobalTransforms(BoneSkeletalMesh->GetRefSkeleton(), ComponentPose);
+						TArray<TArray<int32>> BoneBoundVerts;
+						TArray<TArray<float>> BoneBoundWeights;
+						BoneBoundVerts.SetNum(ComponentPose.Num());
+						BoneBoundWeights.SetNum(ComponentPose.Num());
+
+						uint32 NumSkeletonVertices = PositionVertexBuffer.GetNumVertices();
+
+						Chaos::FReal SphereRadius = (Chaos::FReal)0.;
+
+						Chaos::TVec3<float> CoordMaxs(-FLT_MAX);
+						Chaos::TVec3<float> CoordMins(FLT_MAX);
+						for (int32 i = 0; i < int32(NumSkeletonVertices); i++)
+						{
+							for (int32 j = 0; j < 3; j++)
+							{
+								if (PositionVertexBuffer.VertexPosition(i)[j] > CoordMaxs[j])
+								{
+									CoordMaxs[j] = PositionVertexBuffer.VertexPosition(i)[j];
+								}
+								if (PositionVertexBuffer.VertexPosition(i)[j] < CoordMins[j])
+								{
+									CoordMins[j] = PositionVertexBuffer.VertexPosition(i)[j];
+								}
+							}
+						}
+						Chaos::TVec3<float> CoordDiff = (CoordMaxs - CoordMins) * VertexRadiusRatio;
+						SphereRadius = Chaos::FReal(FGenericPlatformMath::Max(CoordDiff[0], FGenericPlatformMath::Max(CoordDiff[1], CoordDiff[2])));
+
+						TArray<Chaos::TSphere<Chaos::FReal, 3>*> VertexSpherePtrs;
+						TArray<Chaos::TSphere<Chaos::FReal, 3>> VertexSpheres;
+
+						VertexSpheres.Init(Chaos::TSphere<Chaos::FReal, 3>(Chaos::TVec3<Chaos::FReal>(0), SphereRadius), NumSkeletonVertices);
+						VertexSpherePtrs.SetNum(NumSkeletonVertices);
+
+						for (int32 i = 0; i < int32(NumSkeletonVertices); i++)
+						{
+							Chaos::TVec3<Chaos::FReal> SphereCenter(PositionVertexBuffer.VertexPosition(i));
+							Chaos::TSphere<Chaos::FReal, 3> VertexSphere(SphereCenter, SphereRadius);
+							VertexSpheres[i] = Chaos::TSphere<Chaos::FReal, 3>(SphereCenter, SphereRadius);
+							VertexSpherePtrs[i] = &VertexSpheres[i];
+						}
+						Chaos::TBoundingVolumeHierarchy<
+							TArray<Chaos::TSphere<Chaos::FReal, 3>*>,
+							TArray<int32>,
+							Chaos::FReal,
+							3> VertexBVH(VertexSpherePtrs);
+
+						if (SkeletalBindingMode == ESkeletalBindingMode::Dataflow_SkeletalBinding_Kinematic)
+						{
+
+							//avoid particles that are already kinematic:
+							TArray<bool> ParticleIsKinematic;
+							ParticleIsKinematic.Init(false, Vertices->Num());
+							typedef GeometryCollection::Facades::FKinematicBindingFacade FKinematics;
+							FKinematics Kinematics(InCollection);
+
+							// Add Kinematics Node
+							for (int i = Kinematics.NumKinematicBindings() - 1; i >= 0; i--)
+							{
+								FKinematics::FBindingKey Key = Kinematics.GetKinematicBindingKey(i);
+
+								int32 BoneIndex = INDEX_NONE;
+								TArray<int32> BoundVerts;
+								TArray<float> BoundWeights;
+								Kinematics.GetBoneBindings(Key, BoneIndex, BoundVerts, BoundWeights);
+
+								for (int32 vdx : BoundVerts)
+								{
+									ParticleIsKinematic[vdx] = true;
+								}
+							}
+
+
+							TArray<Chaos::TVector<int32, 3>> IndicesArray;
+							for (int32 i = 0; i < Indices->Num(); i++)
+							{
+								Chaos::TVector<int32, 3> CurrentIndices(0);
+								for (int32 j = 0; j < 3; j++)
+								{
+									CurrentIndices[j] = (*Indices)[i][j];
+								}
+								if (CurrentIndices[0] != -1
+									&& CurrentIndices[1] != -1
+									&& CurrentIndices[2] != -1)
+								{
+									IndicesArray.Emplace(CurrentIndices);
+								}
+							}
+							TArray<TArray<int32>> LocalIndex;
+							TArray<TArray<int32>>* LocalIndexPtr = &LocalIndex;
+							TArray<TArray<int>> GlobalIndex = Chaos::Utilities::ComputeIncidentElements(IndicesArray, LocalIndexPtr);
+							int32 ActualParticleCount = 0;
+							for (int32 l = 0; l < GlobalIndex.Num(); l++)
+							{
+								if (GlobalIndex[l].Num() > 0)
+								{
+									ActualParticleCount += 1;
+								}
+							}
+							TArray<Chaos::TVector<float, 3>> IndicesPositions;
+							IndicesPositions.SetNum(ActualParticleCount);
+							TArray<int32> IndicesMap;
+							IndicesMap.SetNum(ActualParticleCount);
+							int32 CurrentParticleIndex = 0;
+							for (int32 i = 0; i < GlobalIndex.Num(); i++)
+							{
+								if (GlobalIndex[i].Num() > 0)
+								{
+									IndicesPositions[CurrentParticleIndex] = (*Vertices)[(*Indices)[GlobalIndex[i][0]][LocalIndex[i][0]]];
+									IndicesMap[CurrentParticleIndex] = (*Indices)[GlobalIndex[i][0]][LocalIndex[i][0]];
+									CurrentParticleIndex += 1;
+								}
+							}
+
+
+							//for (int32 i = 0; i < Vertices->Num(); i++)
+							for (int32 i = 0; i < IndicesMap.Num(); i++)
+							{
+								//only work on particles that are not kinematic:
+								if (!ParticleIsKinematic[IndicesMap[i]])
+								{
+									TArray<int32> ParticleIntersection = VertexBVH.FindAllIntersections((*Vertices)[IndicesMap[i]]);
+									int32 MinIndex = -1;
+									float MinDis = SphereRadius;
+									for (int32 j = 0; j < ParticleIntersection.Num(); j++)
+									{
+										Chaos::FRealSingle CurrentDistance = ((*Vertices)[IndicesMap[i]] - PositionVertexBuffer.VertexPosition(ParticleIntersection[j])).Size();
+										if (CurrentDistance < MinDis)
+										{
+											MinDis = CurrentDistance;
+											MinIndex = ParticleIntersection[j];
+										}
+									}
+
+									if (MinIndex != -1)
+									{
+										int32 BoneParticleIndex = -1;
+
+										int32 SectionIndex;
+										int32 VertIndex;
+										LODRenderData->GetSectionFromVertexIndex(MinIndex, SectionIndex, VertIndex);
+
+										check(SectionIndex < LODRenderData->RenderSections.Num());
+										const FSkelMeshRenderSection& Section = LODRenderData->RenderSections[SectionIndex];
+										int32 BufferVertIndex = Section.GetVertexBufferIndex() + VertIndex;
+										for (int32 InfluenceIndex = 0; InfluenceIndex < MaxBoneInfluences; InfluenceIndex++)
+										{
+											const int32 BoneIndex = Section.BoneMap[SkinWeightVertexBuffer->GetBoneIndex(BufferVertIndex, InfluenceIndex)];
+											const float	Weight = (float)SkinWeightVertexBuffer->GetBoneWeight(BufferVertIndex, InfluenceIndex) * UE::AnimationCore::InvMaxRawBoneWeightFloat;
+											if (Weight > float(0) && 0 <= BoneIndex && BoneIndex < ComponentPose.Num())
+											{
+												BoneBoundVerts[BoneIndex].Add(IndicesMap[i]);
+												//BoneBoundWeights[BoneIndex].Add(Weight);
+												BoneBoundWeights[BoneIndex].Add(1.f);
+												break;
+											}
+										}
+									}
+								}
+							}
+						}
+						else
+						{
+							TMap<int32, int32> BoneVertexToCollectionMap;
+							GeometryCollection::Facades::FPositionTargetFacade PositionTargets(InCollection);
+							PositionTargets.DefineSchema();
+
+							for (int32 i = 0; i < Indices->Num(); i++)
+							{
+								TArray<int32> TriangleIntersections0 = VertexBVH.FindAllIntersections((*Vertices)[(*Indices)[i][0]]);
+								TArray<int32> TriangleIntersections1 = VertexBVH.FindAllIntersections((*Vertices)[(*Indices)[i][1]]);
+								TArray<int32> TriangleIntersections2 = VertexBVH.FindAllIntersections((*Vertices)[(*Indices)[i][2]]);
+								TriangleIntersections0.Sort();
+								TriangleIntersections1.Sort();
+								TriangleIntersections2.Sort();
+								TArray<int32> TriangleIntersections({});
+								for (int32 k = 0; k < TriangleIntersections0.Num(); k++)
+								{
+									if (TriangleIntersections1.Contains(TriangleIntersections0[k])
+										&& TriangleIntersections2.Contains(TriangleIntersections0[k]))
+									{
+										TriangleIntersections.Emplace(TriangleIntersections0[k]);
+									}
+								}
+								int32 MinIndex = -1;
+								float MinDis = SphereRadius;
+								Chaos::TVector<float, 3> ClosestBary(0.f);
+								for (int32 j = 0; j < TriangleIntersections.Num(); j++)
+								{
+									Chaos::TVector<float, 3> Bary,
+										TriPos0((*Vertices)[(*Indices)[i][0]]), TriPos1((*Vertices)[(*Indices)[i][1]]), TriPos2((*Vertices)[(*Indices)[i][2]]),
+										ParticlePos(PositionVertexBuffer.VertexPosition(TriangleIntersections[j]));
+									Chaos::TVector<Chaos::FRealSingle, 3> ClosestPoint = Chaos::FindClosestPointAndBaryOnTriangle(TriPos0, TriPos1, TriPos2, ParticlePos, Bary);
+									Chaos::FRealSingle CurrentDistance = (ParticlePos - ClosestPoint).Size();
+									if (CurrentDistance < MinDis)
+									{
+										MinDis = CurrentDistance;
+										MinIndex = TriangleIntersections[j];
+										ClosestBary = Bary;
+									}
+								}
+
+								if (MinIndex != -1)
+								{
+									int32 BoneParticleIndex = -1;
+									//TODO: add kinematic particles first
+									if (!BoneVertexToCollectionMap.Contains(MinIndex))
+									{
+
+										int32 ParticleIndex = InCollection.AddElements(1, FGeometryCollection::VerticesGroup);
+										TManagedArray<FVector3f>& CurrentVertices = InCollection.ModifyAttribute<FVector3f>("Vertex", FGeometryCollection::VerticesGroup);
+
+										CurrentVertices[ParticleIndex] = PositionVertexBuffer.VertexPosition(MinIndex);
+										BoneVertexToCollectionMap.Emplace(MinIndex, ParticleIndex);
+
+										int32 SectionIndex;
+										int32 VertIndex;
+										LODRenderData->GetSectionFromVertexIndex(MinIndex, SectionIndex, VertIndex);
+
+										check(SectionIndex < LODRenderData->RenderSections.Num());
+										const FSkelMeshRenderSection& Section = LODRenderData->RenderSections[SectionIndex];
+										int32 BufferVertIndex = Section.GetVertexBufferIndex() + VertIndex;
+										for (int32 InfluenceIndex = 0; InfluenceIndex < MaxBoneInfluences; InfluenceIndex++)
+										{
+											const int32 BoneIndex = Section.BoneMap[SkinWeightVertexBuffer->GetBoneIndex(BufferVertIndex, InfluenceIndex)];
+											const float	Weight = (float)SkinWeightVertexBuffer->GetBoneWeight(BufferVertIndex, InfluenceIndex) * UE::AnimationCore::InvMaxRawBoneWeightFloat;
+											if (Weight > float(0) && 0 <= BoneIndex && BoneIndex < ComponentPose.Num())
+											{
+												BoneBoundVerts[BoneIndex].Add(ParticleIndex);
+												BoneBoundWeights[BoneIndex].Add(Weight);
+											}
+										}
+									}
+
+									BoneParticleIndex = BoneVertexToCollectionMap[MinIndex];
+
+									GeometryCollection::Facades::FPositionTargetsData DataPackage;
+									DataPackage.TargetIndex.Init(BoneParticleIndex, 1);
+									DataPackage.TargetWeights.Init(1.f, 1);
+									DataPackage.SourceWeights.Init(1.f, 3);
+									DataPackage.SourceIndex.Init(-1, 3);
+									DataPackage.SourceIndex[0] = (*Indices)[i][0];
+									DataPackage.SourceIndex[1] = (*Indices)[i][1];
+									DataPackage.SourceIndex[2] = (*Indices)[i][2];
+									DataPackage.SourceWeights[0] = ClosestBary[0];
+									DataPackage.SourceWeights[1] = ClosestBary[1];
+									DataPackage.SourceWeights[2] = ClosestBary[2];
+									if (TManagedArray<float>* Mass = InCollection.FindAttribute<float>("Mass", FGeometryCollection::VerticesGroup))
+									{
+										DataPackage.Stiffness = 0.f;
+										for (int32 k = 0; k < 3; k++)
+										{
+											DataPackage.Stiffness += DataPackage.SourceWeights[k] * PositionTargetStiffness * (*Mass)[DataPackage.SourceIndex[k]];
+										}
+										//DataPackage.Stiffness += DataPackage.TargetWeights[0] * PositionTargetStiffness * (*Mass)[DataPackage.TargetIndex[0]];
+										//DataPackage.Stiffness /= 2.f;
+									}
+									else
+									{
+										DataPackage.Stiffness = PositionTargetStiffness;
+									}
+									PositionTargets.AddPositionTarget(DataPackage);
+								}
+
+
+							}
+
+
+						}
+						auto DoubleVert = [](FVector3f V) { return FVector3d(V.X, V.Y, V.Z); };
+
+						for (int32 BoneIndex = 0; BoneIndex < ComponentPose.Num(); ++BoneIndex)
+						{
+							if (BoneBoundVerts[BoneIndex].Num())
+							{
+								FString BoneName = BoneSkeletalMesh->GetRefSkeleton().GetBoneName(BoneIndex).ToString();
+								//get local coords of bound verts
+								typedef GeometryCollection::Facades::FKinematicBindingFacade FKinematics;
+								FKinematics Kinematics(InCollection); Kinematics.DefineSchema();
+								if (Kinematics.IsValid())
+								{
+									FKinematics::FBindingKey Binding = Kinematics.SetBoneBindings(BoneIndex, BoneBoundVerts[BoneIndex], BoneBoundWeights[BoneIndex]);
+									TManagedArray<TArray<FVector3f>>& LocalPos = InCollection.AddAttribute<TArray<FVector3f>>("LocalPosition", Binding.GroupName);
+									Kinematics.AddKinematicBinding(Binding);
+
+									auto FloatVert = [](FVector3d V) { return FVector3f(V.X, V.Y, V.Z); };
+									LocalPos[Binding.Index].SetNum(BoneBoundVerts[BoneIndex].Num());
+									for (int32 i = 0; i < BoneBoundVerts[BoneIndex].Num(); i++)
+									{
+										FVector3f Temp = (*Vertices)[BoneBoundVerts[BoneIndex][i]];
+										LocalPos[Binding.Index][i] = FloatVert(ComponentPose[BoneIndex].InverseTransformPosition(DoubleVert(Temp)));
+									}
+								}
+							}
+						}
+						GeometryCollection::Facades::FVertexBoneWeightsFacade(InCollection).AddBoneWeightsFromKinematicBindings();
+
+					}
+
+
+				}
+			}
+		}
+		SetValue(Context, MoveTemp(InCollection), &Collection);
+	}
+}
\ No newline at end of file
diff --git a/Engine/Plugins/Experimental/ChaosFlesh/Source/ChaosFleshNodes/Private/Dataflow/ChaosFleshSetVertexTetrahedraPositionTargetBindingNode.cpp b/Engine/Plugins/Experimental/ChaosFlesh/Source/ChaosFleshNodes/Private/Dataflow/ChaosFleshSetVertexTetrahedraPositionTargetBindingNode.cpp
new file mode 100644
--- /dev/null
+++ b/Engine/Plugins/Experimental/ChaosFlesh/Source/ChaosFleshNodes/Private/Dataflow/ChaosFleshSetVertexTetrahedraPositionTargetBindingNode.cpp
@@ -0,0 +1,130 @@
+// Copyright Epic Games, Inc. All Rights Reserved.
+
+#include "Dataflow/ChaosFleshSetVertexTetrahedraPositionTargetBindingNode.h"
+
+#include "Chaos/BoundingVolumeHierarchy.h"
+#include "GeometryCollection/Facades/CollectionPositionTargetFacade.h"
+#include "Chaos/Tetrahedron.h"
+#include "ChaosFlesh/TetrahedralCollection.h"
+
+#include UE_INLINE_GENERATED_CPP_BY_NAME(ChaosFleshSetVertexTetrahedraPositionTargetBindingNode)
+
+//DEFINE_LOG_CATEGORY_STATIC(ChaosFleshSetVertexTetrahedraPositionTargetBindingNodeLog, Log, All);
+
+
+void FSetVertexTetrahedraPositionTargetBindingDataflowNode::Evaluate(Dataflow::FContext& Context, const FDataflowOutput* Out) const
+{
+	if (Out->IsA<DataType>(&Collection))
+	{
+		DataType InCollection = GetValue<DataType>(Context, &Collection);
+		if (FindInput(&TargetIndicesIn) && FindInput(&TargetIndicesIn)->GetConnection()) 
+		{
+			if (TManagedArray<FVector3f>* Vertices = InCollection.FindAttribute<FVector3f>("Vertex", FGeometryCollection::VerticesGroup))
+			{
+				if (TManagedArray<FIntVector4>* Tetrahedron = InCollection.FindAttribute<FIntVector4>(FTetrahedralCollection::TetrahedronAttribute, FTetrahedralCollection::TetrahedralGroup))
+				{
+					TManagedArray<int32>* TetrahedronStart = InCollection.FindAttribute<int32>(FTetrahedralCollection::TetrahedronStartAttribute, FGeometryCollection::GeometryGroup);
+					TManagedArray<int32>* TetrahedronCount = InCollection.FindAttribute<int32>(FTetrahedralCollection::TetrahedronCountAttribute, FGeometryCollection::GeometryGroup);
+					if (TetrahedronStart && TetrahedronCount)
+					{	
+						TArray<FString> GeometryGroupGuidsLocal;
+						if (FindInput(&GeometryGroupGuidsIn) && FindInput(&GeometryGroupGuidsIn)->GetConnection())
+						{
+							GeometryGroupGuidsLocal = GetValue<TArray<FString>>(Context, &GeometryGroupGuidsIn);
+						}
+						TManagedArray<FString>* Guids = InCollection.FindAttribute<FString>("Guid", FGeometryCollection::GeometryGroup);
+						for (int32 TetMeshIdx = 0; TetMeshIdx < TetrahedronStart->Num(); TetMeshIdx++)
+						{
+							if (GeometryGroupGuidsLocal.Num() && Guids)
+							{
+								if (!GeometryGroupGuidsLocal.Contains((*Guids)[TetMeshIdx]))
+								{
+									continue;
+								}
+							}
+							const int32 TetMeshStart = (*TetrahedronStart)[TetMeshIdx];
+							const int32 TetMeshCount = (*TetrahedronCount)[TetMeshIdx];
+
+							// Build Tetrahedra
+							TArray<Chaos::TTetrahedron<Chaos::FReal>> Tets;			// Index 0 == TetMeshStart
+							TArray<Chaos::TTetrahedron<Chaos::FReal>*> BVHTetPtrs;
+							Tets.SetNumUninitialized(TetMeshCount);
+							BVHTetPtrs.SetNumUninitialized(TetMeshCount);
+							for (int32 i = 0; i < TetMeshCount; i++)
+							{
+								const int32 Idx = TetMeshStart + i;
+								const FIntVector4& Tet = (*Tetrahedron)[Idx];
+								Tets[i] = Chaos::TTetrahedron<Chaos::FReal>(
+									(*Vertices)[Tet[0]],
+									(*Vertices)[Tet[1]],
+									(*Vertices)[Tet[2]],
+									(*Vertices)[Tet[3]]);
+								BVHTetPtrs[i] = &Tets[i];
+							}
+
+							Chaos::TBoundingVolumeHierarchy<
+								TArray<Chaos::TTetrahedron<Chaos::FReal>*>,
+								TArray<int32>,
+								Chaos::FReal,
+								3> TetBVH(BVHTetPtrs);
+
+							TArray<int32> TargetIndicesLocal = GetValue<TArray<int32>>(Context, &TargetIndicesIn);
+							TArray<int32> SourceIndices;
+							SourceIndices.Init(-1, TargetIndicesLocal.Num());
+
+							GeometryCollection::Facades::FPositionTargetFacade PositionTargets(InCollection);
+							PositionTargets.DefineSchema();
+
+							for (int32 i = 0; i < TargetIndicesLocal.Num(); i++)
+							{
+								if (TargetIndicesLocal[i] > -1 && TargetIndicesLocal[i] < Vertices->Num())
+								{
+									FVector3f ParticlePos = (*Vertices)[TargetIndicesLocal[i]];
+									TArray<int32> TetIntersections = TetBVH.FindAllIntersections(ParticlePos);
+									for (int32 j = 0; j < TetIntersections.Num(); j++)
+									{
+										const int32 TetIdx = TetIntersections[j];
+										if (!Tets[TetIdx].Outside(ParticlePos, 0))
+										{
+											Chaos::TVector<Chaos::FReal, 4> WeightsD = Tets[TetIdx].GetBarycentricCoordinates(ParticlePos);
+											GeometryCollection::Facades::FPositionTargetsData DataPackage;
+											DataPackage.TargetIndex.Init(TargetIndicesLocal[i], 1);
+											DataPackage.TargetWeights.Init(1.f, 1);
+											DataPackage.SourceWeights.Init(1.f, 4);
+											DataPackage.SourceIndex.Init(-1, 4);
+											DataPackage.SourceIndex[0] = (*Tetrahedron)[TetIdx][0];
+											DataPackage.SourceIndex[1] = (*Tetrahedron)[TetIdx][1];
+											DataPackage.SourceIndex[2] = (*Tetrahedron)[TetIdx][2];
+											DataPackage.SourceIndex[3] = (*Tetrahedron)[TetIdx][3];
+											DataPackage.SourceWeights[0] = WeightsD[0];
+											DataPackage.SourceWeights[1] = WeightsD[1];
+											DataPackage.SourceWeights[2] = WeightsD[2];
+											DataPackage.SourceWeights[3] = WeightsD[3];
+											if (TManagedArray<float>* Mass = InCollection.FindAttribute<float>("Mass", FGeometryCollection::VerticesGroup))
+											{
+												DataPackage.Stiffness = 0.f;
+												for (int32 k = 0; k < 4; k++)
+												{
+													DataPackage.Stiffness += DataPackage.SourceWeights[k] * PositionTargetStiffness * (*Mass)[DataPackage.SourceIndex[0]];
+												}
+											}
+											else
+											{
+												DataPackage.Stiffness = PositionTargetStiffness;
+											}
+											PositionTargets.AddPositionTarget(DataPackage);
+											break;
+										}
+									}
+								}
+							}
+						}
+					}
+				}
+			}
+		}
+		SetValue(Context, MoveTemp(InCollection), &Collection);
+	}
+
+
+}
diff --git a/Engine/Plugins/Experimental/ChaosFlesh/Source/ChaosFleshNodes/Private/Dataflow/ChaosFleshSetVertexTrianglePositionTargetBindingNode.cpp b/Engine/Plugins/Experimental/ChaosFlesh/Source/ChaosFleshNodes/Private/Dataflow/ChaosFleshSetVertexTrianglePositionTargetBindingNode.cpp
new file mode 100644
--- /dev/null
+++ b/Engine/Plugins/Experimental/ChaosFlesh/Source/ChaosFleshNodes/Private/Dataflow/ChaosFleshSetVertexTrianglePositionTargetBindingNode.cpp
@@ -0,0 +1,185 @@
+// Copyright Epic Games, Inc. All Rights Reserved.
+
+#include "Dataflow/ChaosFleshSetVertexTrianglePositionTargetBindingNode.h"
+
+#include "Chaos/BoundingVolumeHierarchy.h"
+#include "GeometryCollection/Facades/CollectionMeshFacade.h"
+#include "GeometryCollection/Facades/CollectionPositionTargetFacade.h"
+#include "ChaosFlesh/TetrahedralCollection.h"
+
+#include UE_INLINE_GENERATED_CPP_BY_NAME(ChaosFleshSetVertexTrianglePositionTargetBindingNode)
+
+//DEFINE_LOG_CATEGORY_STATIC(ChaosFleshSetVertexTrianglePositionTargetBindingNodeLog, Log, All);
+
+void FSetVertexTrianglePositionTargetBindingDataflowNode::Evaluate(Dataflow::FContext& Context, const FDataflowOutput* Out) const
+{
+	if (Out->IsA<DataType>(&Collection))
+	{
+		DataType InCollection = GetValue<DataType>(Context, &Collection);
+
+		if (TManagedArray<FVector3f>* Vertices = InCollection.FindAttribute<FVector3f>("Vertex", FGeometryCollection::VerticesGroup))
+			{
+				if (TManagedArray<FIntVector>* Indices = InCollection.FindAttribute<FIntVector>("Indices", FGeometryCollection::FacesGroup))
+				{
+					GeometryCollection::Facades::FCollectionMeshFacade MeshFacade(InCollection);
+					TArray<int32> ComponentIndex = MeshFacade.GetGeometryGroupIndexArray();
+					TArray<Chaos::TVector<int32, 3>> IndicesArray;
+					for (int32 i = 0; i < Indices->Num(); i++)
+					{
+						Chaos::TVector<int32, 3> CurrentIndices(0);
+						for (int32 j = 0; j < 3; j++) 
+						{
+							CurrentIndices[j] = (*Indices)[i][j];
+						}
+						if (CurrentIndices[0] != -1
+							&& CurrentIndices[1] != -1
+							&& CurrentIndices[2] != -1)
+						{
+							IndicesArray.Emplace(CurrentIndices);
+						}
+					}
+					TArray<TArray<int32>> LocalIndex;
+					TArray<TArray<int32>>* LocalIndexPtr = &LocalIndex;
+					TArray<TArray<int>> GlobalIndex = Chaos::Utilities::ComputeIncidentElements(IndicesArray, LocalIndexPtr);
+					int32 ActualParticleCount = 0;
+					for (int32 l = 0; l < GlobalIndex.Num(); l++)
+					{
+						if (GlobalIndex[l].Num() > 0)
+						{
+							ActualParticleCount += 1;
+						}
+					}
+					TArray<Chaos::TVector<float, 3>> IndicesPositions; 
+					IndicesPositions.SetNum(ActualParticleCount);
+					TArray<int32> IndicesMap;
+					IndicesMap.SetNum(ActualParticleCount);
+					int32 CurrentParticleIndex = 0;
+					for (int32 i = 0; i < GlobalIndex.Num(); i++)
+					{
+						if (GlobalIndex[i].Num() > 0)
+						{
+							IndicesPositions[CurrentParticleIndex] = (*Vertices)[(*Indices)[GlobalIndex[i][0]][LocalIndex[i][0]]];
+							IndicesMap[CurrentParticleIndex] = (*Indices)[GlobalIndex[i][0]][LocalIndex[i][0]];
+							CurrentParticleIndex += 1;
+						}
+					}
+					Chaos::FReal SphereRadius = (Chaos::FReal)0.;
+
+					Chaos::TVec3<float> CoordMaxs(-FLT_MAX);
+					Chaos::TVec3<float> CoordMins(FLT_MAX);
+					for (int32 i = 0; i < IndicesPositions.Num(); i++)
+					{
+						for (int32 j = 0; j < 3; j++)
+						{
+							if (IndicesPositions[i][j] > CoordMaxs[j])
+							{
+								CoordMaxs[j] = IndicesPositions[i][j];
+							}
+							if (IndicesPositions[i][j] < CoordMins[j])
+							{
+								CoordMins[j] = IndicesPositions[i][j];
+							}
+						}
+					}
+					Chaos::TVec3<float> CoordDiff = (CoordMaxs - CoordMins) * VertexRadiusRatio;
+					SphereRadius = Chaos::FReal(FGenericPlatformMath::Min(CoordDiff[0], FGenericPlatformMath::Min(CoordDiff[1], CoordDiff[2])));
+
+					TArray<Chaos::TSphere<Chaos::FReal, 3>*> VertexSpherePtrs;
+					TArray<Chaos::TSphere<Chaos::FReal, 3>> VertexSpheres;
+
+					VertexSpheres.Init(Chaos::TSphere<Chaos::FReal, 3>(Chaos::TVec3<Chaos::FReal>(0), SphereRadius), IndicesPositions.Num());
+					VertexSpherePtrs.SetNum(IndicesPositions.Num());
+
+					for (int32 i = 0; i < IndicesPositions.Num(); i++)
+					{
+						Chaos::TVec3<Chaos::FReal> SphereCenter(IndicesPositions[i]);
+						Chaos::TSphere<Chaos::FReal, 3> VertexSphere(SphereCenter, SphereRadius);
+						VertexSpheres[i] = Chaos::TSphere<Chaos::FReal, 3>(SphereCenter, SphereRadius);
+						VertexSpherePtrs[i] = &VertexSpheres[i];
+					}
+					Chaos::TBoundingVolumeHierarchy<
+						TArray<Chaos::TSphere<Chaos::FReal, 3>*>,
+						TArray<int32>,
+						Chaos::FReal,
+						3> VertexBVH(VertexSpherePtrs);
+
+					GeometryCollection::Facades::FPositionTargetFacade PositionTargets(InCollection);
+					PositionTargets.DefineSchema();
+
+					for (int32 i = 0; i < Indices->Num(); i++)
+					{
+						TArray<int32> TriangleIntersections0 = VertexBVH.FindAllIntersections((*Vertices)[(*Indices)[i][0]]);
+						TArray<int32> TriangleIntersections1 = VertexBVH.FindAllIntersections((*Vertices)[(*Indices)[i][1]]);
+						TArray<int32> TriangleIntersections2 = VertexBVH.FindAllIntersections((*Vertices)[(*Indices)[i][2]]);
+						TriangleIntersections0.Sort();
+						TriangleIntersections1.Sort();
+						TriangleIntersections2.Sort();
+
+						TArray<int32> TriangleIntersections({});
+						for (int32 k = 0; k < TriangleIntersections0.Num(); k++)
+						{
+							if (TriangleIntersections1.Contains(TriangleIntersections0[k]) 
+								&& TriangleIntersections2.Contains(TriangleIntersections0[k]))
+							{
+								TriangleIntersections.Emplace(TriangleIntersections0[k]);
+							}
+						}
+
+						int32 TriangleIndex = ComponentIndex[(*Indices)[i][0]];
+						int32 MinIndex = -1;
+						float MinDis = SphereRadius;
+						Chaos::TVector<float, 3> ClosestBary(0.f);
+						for (int32 j = 0; j < TriangleIntersections.Num(); j++)
+						{
+							if (ComponentIndex[IndicesMap[TriangleIntersections[j]]] >= 0 && TriangleIndex >= 0 && ComponentIndex[IndicesMap[TriangleIntersections[j]]] != TriangleIndex)
+							{
+								Chaos::TVector<float, 3> Bary, TriPos0((*Vertices)[(*Indices)[i][0]]), TriPos1((*Vertices)[(*Indices)[i][1]]), TriPos2((*Vertices)[(*Indices)[i][2]]), ParticlePos((*Vertices)[IndicesMap[TriangleIntersections[j]]]);
+								Chaos::TVector<Chaos::FRealSingle, 3> ClosestPoint = Chaos::FindClosestPointAndBaryOnTriangle(TriPos0, TriPos1, TriPos2, ParticlePos, Bary);
+								Chaos::FRealSingle CurrentDistance = ((*Vertices)[IndicesMap[TriangleIntersections[j]]] - ClosestPoint).Size();
+								if (CurrentDistance < MinDis)
+								{
+									MinDis = CurrentDistance;
+									MinIndex = IndicesMap[TriangleIntersections[j]];
+									ClosestBary = Bary;
+								}
+
+							}
+						}
+						if (MinIndex != -1
+							&& MinIndex != (*Indices)[i][0]
+							&& MinIndex != (*Indices)[i][1]
+							&& MinIndex != (*Indices)[i][2])
+						{
+							GeometryCollection::Facades::FPositionTargetsData DataPackage;
+							DataPackage.TargetIndex.Init(MinIndex, 1);
+							DataPackage.TargetWeights.Init(1.f, 1);
+							DataPackage.SourceWeights.Init(1.f, 3);
+							DataPackage.SourceIndex.Init(-1, 3);
+							DataPackage.SourceIndex[0] = (*Indices)[i][0];
+							DataPackage.SourceIndex[1] = (*Indices)[i][1];
+							DataPackage.SourceIndex[2] = (*Indices)[i][2];
+							DataPackage.SourceWeights[0] = ClosestBary[0];
+							DataPackage.SourceWeights[1] = ClosestBary[1];
+							DataPackage.SourceWeights[2] = ClosestBary[2];
+							if (TManagedArray<float>* Mass = InCollection.FindAttribute<float>("Mass", FGeometryCollection::VerticesGroup))
+							{
+								DataPackage.Stiffness = 0.f;
+								for (int32 k = 0; k < 3; k++)
+								{
+									DataPackage.Stiffness += DataPackage.SourceWeights[k] * PositionTargetStiffness * (*Mass)[DataPackage.SourceIndex[k]];
+								}
+								DataPackage.Stiffness += DataPackage.TargetWeights[0] * PositionTargetStiffness * (*Mass)[DataPackage.TargetIndex[0]];
+								DataPackage.Stiffness /= 2.f;
+							}
+							else
+							{
+								DataPackage.Stiffness = PositionTargetStiffness;
+							}
+							PositionTargets.AddPositionTarget(DataPackage);
+						}
+					}
+				}
+		}
+		SetValue(Context, MoveTemp(InCollection), &Collection);
+	}
+}
diff --git a/Engine/Plugins/Experimental/ChaosFlesh/Source/ChaosFleshNodes/Private/Dataflow/ChaosFleshSetVertexVertexPositionTargetBindingNode.cpp b/Engine/Plugins/Experimental/ChaosFlesh/Source/ChaosFleshNodes/Private/Dataflow/ChaosFleshSetVertexVertexPositionTargetBindingNode.cpp
new file mode 100644
--- /dev/null
+++ b/Engine/Plugins/Experimental/ChaosFlesh/Source/ChaosFleshNodes/Private/Dataflow/ChaosFleshSetVertexVertexPositionTargetBindingNode.cpp
@@ -0,0 +1,129 @@
+// Copyright Epic Games, Inc. All Rights Reserved.
+
+#include "Dataflow/ChaosFleshSetVertexVertexPositionTargetBindingNode.h"
+
+#include "Chaos/BoundingVolumeHierarchy.h"
+#include "GeometryCollection/Facades/CollectionPositionTargetFacade.h"
+
+#include UE_INLINE_GENERATED_CPP_BY_NAME(ChaosFleshSetVertexVertexPositionTargetBindingNode)
+
+//DEFINE_LOG_CATEGORY_STATIC(ChaosFleshSetVertexVertexPositionTargetBindingNodeLog, Log, All);
+
+
+
+void FSetVertexVertexPositionTargetBindingDataflowNode::Evaluate(Dataflow::FContext& Context, const FDataflowOutput* Out) const
+{
+	if (Out->IsA<DataType>(&Collection))
+	{
+		DataType InCollection = GetValue<DataType>(Context, &Collection);
+
+		if (TManagedArray<FVector3f>* Vertices = InCollection.FindAttribute<FVector3f>("Vertex", FGeometryCollection::VerticesGroup))
+		{
+			if (FindInput(&TargetIndicesIn) && FindInput(&TargetIndicesIn)->GetConnection())
+			{
+				Chaos::FReal SphereRadius = (Chaos::FReal)0.;
+
+				Chaos::TVec3<float> CoordMaxs(-FLT_MAX);
+				Chaos::TVec3<float> CoordMins(FLT_MAX);
+				for (int32 i = 0; i < Vertices->Num(); i++)
+				{
+					for (int32 j = 0; j < 3; j++) 
+					{
+						if ((*Vertices)[i][j] > CoordMaxs[j]) 
+						{
+							CoordMaxs[j] = (*Vertices)[i][j];
+						}
+						if ((*Vertices)[i][j] < CoordMins[j])
+						{
+							CoordMins[j] = (*Vertices)[i][j];
+						}
+					}
+				}
+
+				Chaos::TVec3<float> CoordDiff = (CoordMaxs - CoordMins) * RadiusRatio;
+
+				SphereRadius = Chaos::FReal(FGenericPlatformMath::Min(CoordDiff[0], FGenericPlatformMath::Min(CoordDiff[1], CoordDiff[2])));
+
+				TArray<Chaos::TSphere<Chaos::FReal, 3>*> VertexSpherePtrs;
+				TArray<Chaos::TSphere<Chaos::FReal, 3>> VertexSpheres;
+
+				VertexSpheres.Init(Chaos::TSphere<Chaos::FReal, 3>(Chaos::TVec3<Chaos::FReal>(0), SphereRadius), Vertices->Num());
+				VertexSpherePtrs.SetNum(Vertices->Num());
+
+				for (int32 i = 0; i < Vertices->Num(); i++)
+				{
+					Chaos::TVec3<Chaos::FReal> SphereCenter((*Vertices)[i]);
+					Chaos::TSphere<Chaos::FReal, 3> VertexSphere(SphereCenter, SphereRadius);
+					VertexSpheres[i] = Chaos::TSphere<Chaos::FReal, 3>(SphereCenter, SphereRadius);
+					VertexSpherePtrs[i] = &VertexSpheres[i];
+				}
+				Chaos::TBoundingVolumeHierarchy<
+					TArray<Chaos::TSphere<Chaos::FReal, 3>*>,
+					TArray<int32>,
+					Chaos::FReal,
+					3> VertexBVH(VertexSpherePtrs);
+
+			
+				TArray<int32> TargetIndicesLocal = GetValue<TArray<int32>>(Context, &TargetIndicesIn);
+				TSet<int32> TargetIndicesSet = TSet<int32>(TargetIndicesLocal);
+				TArray<int32> SourceIndices;
+				SourceIndices.Init(-1, TargetIndicesLocal.Num());
+				
+				for (int32 i = 0; i < TargetIndicesLocal.Num(); i++)
+				{
+					if (TargetIndicesLocal[i] > -1 && TargetIndicesLocal[i] < Vertices->Num()) 
+					{
+						FVector3f ParticlePos = (*Vertices)[TargetIndicesLocal[i]];
+						TArray<int32> VertexIntersections = VertexBVH.FindAllIntersections(ParticlePos);
+						float MinDistance = 10.f * (float)SphereRadius;
+						int32 MinIndex = -1;
+						for (int32 k = 0; k < VertexIntersections.Num(); k++)
+						{
+							if ((ParticlePos - (*Vertices)[VertexIntersections[k]]).Size() < MinDistance
+								&& VertexIntersections[k] != TargetIndicesLocal[i] && !TargetIndicesSet.Contains(VertexIntersections[k]))
+							{
+								MinIndex = VertexIntersections[k];
+								MinDistance = (ParticlePos - (*Vertices)[VertexIntersections[k]]).Size();
+							}
+						}
+						ensure(MinIndex != -1);
+						SourceIndices[i] = MinIndex;
+					}
+				}
+				GeometryCollection::Facades::FPositionTargetFacade PositionTargets(InCollection);
+				PositionTargets.DefineSchema();
+				for (int32 i = 0; i < TargetIndicesLocal.Num(); i++)
+				{
+					if (SourceIndices[i] != -1)
+					{
+						GeometryCollection::Facades::FPositionTargetsData DataPackage;
+						DataPackage.TargetIndex.Init(TargetIndicesLocal[i], 1);
+						DataPackage.TargetWeights.Init(1.f, 1);
+						DataPackage.SourceWeights.Init(1.f, 1);
+						DataPackage.SourceIndex.Init(SourceIndices[i], 1);
+						if (TManagedArray<float>* Mass = InCollection.FindAttribute<float>("Mass", FGeometryCollection::VerticesGroup))
+						{
+							if ((*Mass)[SourceIndices[i]] > 0.f)
+							{
+								DataPackage.Stiffness = PositionTargetStiffness * (*Mass)[SourceIndices[i]];
+							}
+							else
+							{
+								DataPackage.Stiffness = PositionTargetStiffness;
+							}
+						}
+						else
+						{
+							DataPackage.Stiffness = PositionTargetStiffness;
+						}
+
+						PositionTargets.AddPositionTarget(DataPackage);
+					}
+				}
+			}
+			
+		}
+		SetValue(Context, MoveTemp(InCollection), &Collection);
+	}
+}
+
diff --git a/Engine/Plugins/Experimental/ChaosFlesh/Source/ChaosFleshNodes/Private/Dataflow/ChaosFleshVisualizeFiberFieldNode.cpp b/Engine/Plugins/Experimental/ChaosFlesh/Source/ChaosFleshNodes/Private/Dataflow/ChaosFleshVisualizeFiberFieldNode.cpp
new file mode 100644
--- /dev/null
+++ b/Engine/Plugins/Experimental/ChaosFlesh/Source/ChaosFleshNodes/Private/Dataflow/ChaosFleshVisualizeFiberFieldNode.cpp
@@ -0,0 +1,40 @@
+// Copyright Epic Games, Inc. All Rights Reserved.
+
+#include "Dataflow/ChaosFleshVisualizeFiberFieldNode.h"
+
+#include "ChaosFlesh/TetrahedralCollection.h"
+
+#include UE_INLINE_GENERATED_CPP_BY_NAME(ChaosFleshVisualizeFiberFieldNode)
+
+void FVisualizeFiberFieldNode::Evaluate(Dataflow::FContext& Context, const FDataflowOutput* Out) const
+{
+	if (Out->IsA<FFieldCollection>(&VectorField))
+	{
+		FManagedArrayCollection InCollection = GetValue<FManagedArrayCollection>(Context, &Collection);
+		FFieldCollection OutVectorField = VectorField;
+		
+		if (TManagedArray<FVector3f>* Vertex = InCollection.FindAttribute<FVector3f>("Vertex", "Vertices"))
+		{
+			if (TManagedArray<FIntVector4>* Elements = InCollection.FindAttribute<FIntVector4>(FTetrahedralCollection::TetrahedronAttribute, FTetrahedralCollection::TetrahedralGroup))
+			{
+				if (TManagedArray<FVector3f>* FiberDirections = InCollection.FindAttribute<FVector3f>("FiberDirection", FTetrahedralCollection::TetrahedralGroup))
+				{
+					ensureMsgf(Elements->Num() == FiberDirections->Num(), TEXT("Fiber direction has different size than elements"));
+					for (int32 ElemIndex = 0; ElemIndex < Elements->Num(); ElemIndex++)
+					{
+						FVector3f VectorStart = { 0,0,0 };
+						for (int32 LocalIndex = 0; LocalIndex < 4; LocalIndex++)
+						{
+							VectorStart += (*Vertex)[(*Elements)[ElemIndex][LocalIndex]];
+						}
+						VectorStart /= float(4);
+						FVector3f VectorEnd = VectorStart + (*FiberDirections)[ElemIndex] * VectorScale;
+						OutVectorField.AddVectorToField(VectorStart, VectorEnd);
+					}
+				}
+			}
+		}
+		
+		Out->SetValue(MoveTemp(OutVectorField), Context);
+	}
+}
diff --git a/Engine/Plugins/Experimental/ChaosFlesh/Source/ChaosFleshNodes/Public/Dataflow/ChaosFleshAddKinematicParticlesNode.h b/Engine/Plugins/Experimental/ChaosFlesh/Source/ChaosFleshNodes/Public/Dataflow/ChaosFleshAddKinematicParticlesNode.h
new file mode 100644
--- /dev/null
+++ b/Engine/Plugins/Experimental/ChaosFlesh/Source/ChaosFleshNodes/Public/Dataflow/ChaosFleshAddKinematicParticlesNode.h
@@ -0,0 +1,63 @@
+// Copyright Epic Games, Inc. All Rights Reserved.
+
+#pragma once 
+
+#include "CoreMinimal.h"
+#include "Dataflow/ChaosFleshNodesUtility.h"
+#include "Dataflow/DataflowCore.h"
+#include "Dataflow/DataflowEngine.h"
+
+#include "ChaosFleshAddKinematicParticlesNode.generated.h"
+
+class USkeletalMesh;
+
+USTRUCT(meta = (DataflowFlesh))
+struct FAddKinematicParticlesDataflowNode : public FDataflowNode
+{
+	GENERATED_USTRUCT_BODY()
+	DATAFLOW_NODE_DEFINE_INTERNAL(FAddKinematicParticlesDataflowNode, "AddKinematicParticles", "Flesh", "")
+	DATAFLOW_NODE_RENDER_TYPE(FGeometryCollection::StaticType(), "Collection")
+
+public:
+	typedef FManagedArrayCollection DataType;
+
+	UPROPERTY(EditAnywhere, Category = "Dataflow")
+		float Radius = 40.f;
+
+	UPROPERTY(EditAnywhere, Category = "Dataflow")
+		FTransform Transform;
+
+	UPROPERTY(EditAnywhere, Category = "Dataflow")
+		ESkeletalSeletionMode SkeletalSelectionMode = ESkeletalSeletionMode::Dataflow_SkeletalSelection_Single;
+
+	UPROPERTY(meta = (DataflowInput, DataflowOutput, DisplayName = "Collection", DataflowPassthrough = "Collection"))
+		FManagedArrayCollection Collection;
+
+	UPROPERTY(meta = (DataflowInput, DisplayName = "SkeletalMesh"))
+		TObjectPtr<USkeletalMesh> SkeletalMeshIn;
+
+	UPROPERTY(meta = (DataflowInput, DisplayName = "SelectionSet"))
+		TArray<int32> VertexIndicesIn;
+
+	UPROPERTY(meta = (DataflowInput, DisplayName = "BoneIndex"))
+		int32 BoneIndexIn = 0;
+
+	UPROPERTY(meta = (DataflowOutput, DisplayName = "TargetIndices"))
+		TArray<int32> TargetIndicesOut;
+
+
+
+	FAddKinematicParticlesDataflowNode(const Dataflow::FNodeParameters& InParam, FGuid InGuid = FGuid::NewGuid())
+		: FDataflowNode(InParam, InGuid)
+	{
+		RegisterInputConnection(&Collection);
+		RegisterOutputConnection(&Collection, &Collection);
+		RegisterInputConnection(&SkeletalMeshIn);
+		RegisterInputConnection(&VertexIndicesIn);
+		RegisterInputConnection(&BoneIndexIn);
+		RegisterOutputConnection(&TargetIndicesOut);
+	}
+
+	virtual void Evaluate(Dataflow::FContext& Context, const FDataflowOutput* Out) const override;
+};
+
diff --git a/Engine/Plugins/Experimental/ChaosFlesh/Source/ChaosFleshNodes/Public/Dataflow/ChaosFleshBindForRenderToSkeletalMeshNode.h b/Engine/Plugins/Experimental/ChaosFlesh/Source/ChaosFleshNodes/Public/Dataflow/ChaosFleshBindForRenderToSkeletalMeshNode.h
deleted file mode 100644
diff --git a/Engine/Plugins/Experimental/ChaosFlesh/Source/ChaosFleshNodes/Public/Dataflow/ChaosFleshComputeFiberFieldNode.h b/Engine/Plugins/Experimental/ChaosFlesh/Source/ChaosFleshNodes/Public/Dataflow/ChaosFleshComputeFiberFieldNode.h
new file mode 100644
--- /dev/null
+++ b/Engine/Plugins/Experimental/ChaosFlesh/Source/ChaosFleshNodes/Public/Dataflow/ChaosFleshComputeFiberFieldNode.h
@@ -0,0 +1,72 @@
+// Copyright Epic Games, Inc. All Rights Reserved.
+
+#pragma once 
+
+#include "CoreMinimal.h"
+#include "Dataflow/DataflowCore.h"
+#include "Dataflow/DataflowEngine.h"
+
+#include "ChaosFleshComputeFiberFieldNode.generated.h"
+
+
+
+/**
+* Computes a muscle fiber direction per tetrahedron from a GeometryCollection containing tetrahedra, 
+* vertices, and origin & insertion vertex fields.  Fiber directions should smoothly follow the geometry
+* oriented from the origin vertices pointing to the insertion vertices.
+*/
+USTRUCT(meta = (DataflowFlesh))
+struct FComputeFiberFieldNode : public FDataflowNode
+{
+	GENERATED_USTRUCT_BODY()
+	DATAFLOW_NODE_DEFINE_INTERNAL(FComputeFiberFieldNode, "ComputeFiberField", "Flesh", "")
+	DATAFLOW_NODE_RENDER_TYPE(FGeometryCollection::StaticType(), "Collection")
+
+public:
+	//typedef FManagedArrayCollection DataType;
+
+	UPROPERTY(meta = (DataflowInput, DataflowOutput, DisplayName = "Collection", DataflowPassthrough  = "Collection"))
+	FManagedArrayCollection Collection;
+
+	UPROPERTY(meta = (DataflowInput, DisplayName = "OriginVertexIndices"))
+	TArray<int32> OriginIndices;
+
+	UPROPERTY(meta = (DataflowInput, DisplayName = "InsertionVertexIndices"))
+	TArray<int32> InsertionIndices;
+
+	UPROPERTY(EditAnywhere, Category = "Dataflow")
+	FString OriginInsertionGroupName = FString();
+
+	UPROPERTY(EditAnywhere, Category = "Dataflow")
+	FString OriginVertexFieldName = FString("Origin");
+
+	UPROPERTY(EditAnywhere, Category = "Dataflow")
+	FString InsertionVertexFieldName = FString("Insertion");
+
+	UPROPERTY(EditAnywhere, Category = "Dataflow")
+	int32 MaxIterations = 100;
+
+	UPROPERTY(EditAnywhere, Category = "Dataflow")
+	float Tolerance = 1.0e-7;
+
+	FComputeFiberFieldNode(const Dataflow::FNodeParameters& InParam, FGuid InGuid = FGuid::NewGuid())
+		: FDataflowNode(InParam, InGuid)
+	{
+		RegisterInputConnection(&Collection);
+		RegisterInputConnection(&OriginIndices);
+		RegisterInputConnection(&InsertionIndices);
+		RegisterOutputConnection(&Collection, &Collection);
+	}
+
+	virtual void Evaluate(Dataflow::FContext& Context, const FDataflowOutput* Out) const override;
+
+	TArray<int32> GetNonZeroIndices(const TArray<uint8>& Map) const;
+
+	TArray<FVector3f> ComputeFiberField(
+		const TManagedArray<FIntVector4>& Elements,
+		const TManagedArray<FVector3f>& Vertex,
+		const TManagedArray<TArray<int32>>& IncidentElements,
+		const TManagedArray<TArray<int32>>& IncidentElementsLocalIndex,
+		const TArray<int32>& Origin,
+		const TArray<int32>& Insertion) const;
+};
diff --git a/Engine/Plugins/Experimental/ChaosFlesh/Source/ChaosFleshNodes/Public/Dataflow/ChaosFleshComputeIslandsNode.h b/Engine/Plugins/Experimental/ChaosFlesh/Source/ChaosFleshNodes/Public/Dataflow/ChaosFleshComputeIslandsNode.h
new file mode 100644
--- /dev/null
+++ b/Engine/Plugins/Experimental/ChaosFlesh/Source/ChaosFleshNodes/Public/Dataflow/ChaosFleshComputeIslandsNode.h
@@ -0,0 +1,32 @@
+// Copyright Epic Games, Inc. All Rights Reserved.
+
+#pragma once 
+
+#include "CoreMinimal.h"
+#include "Dataflow/DataflowCore.h"
+#include "Dataflow/DataflowEngine.h"
+
+#include "ChaosFleshComputeIslandsNode.generated.h"
+
+//@todo(deprecate)
+USTRUCT(meta = (DataflowFlesh))
+struct FComputeIslandsNode : public FDataflowNode
+{
+	GENERATED_USTRUCT_BODY()
+	DATAFLOW_NODE_DEFINE_INTERNAL(FComputeIslandsNode, "ComputeIslands", "Flesh", "")
+	DATAFLOW_NODE_RENDER_TYPE(FGeometryCollection::StaticType(), "Collection")
+
+public:
+
+	UPROPERTY(meta = (DataflowInput, DataflowOutput, DisplayName = "Collection", DataflowPassthrough = "Collection"))
+	FManagedArrayCollection Collection;
+
+	FComputeIslandsNode(const Dataflow::FNodeParameters& InParam, FGuid InGuid = FGuid::NewGuid())
+		: FDataflowNode(InParam, InGuid)
+	{
+		RegisterInputConnection(&Collection);
+		RegisterOutputConnection(&Collection, &Collection);
+	}
+
+	virtual void Evaluate(Dataflow::FContext& Context, const FDataflowOutput* Out) const override;
+};
diff --git a/Engine/Plugins/Experimental/ChaosFlesh/Source/ChaosFleshNodes/Public/Dataflow/ChaosFleshComputeMuscleActivationNode.h b/Engine/Plugins/Experimental/ChaosFlesh/Source/ChaosFleshNodes/Public/Dataflow/ChaosFleshComputeMuscleActivationNode.h
new file mode 100644
--- /dev/null
+++ b/Engine/Plugins/Experimental/ChaosFlesh/Source/ChaosFleshNodes/Public/Dataflow/ChaosFleshComputeMuscleActivationNode.h
@@ -0,0 +1,41 @@
+// Copyright Epic Games, Inc. All Rights Reserved.
+
+#pragma once 
+
+#include "CoreMinimal.h"
+#include "Dataflow/DataflowCore.h"
+#include "Dataflow/DataflowEngine.h"
+
+#include "ChaosFleshComputeMuscleActivationNode.generated.h"
+
+/**
+* Computes an orthogonal matrix for each element
+* M = [v,w,u], where v is the fiber direction of that element, w and u are chosen to be orthogonal to v and each other.
+*/
+USTRUCT(meta = (DataflowFlesh))
+struct FComputeMuscleActivationDataNode : public FDataflowNode
+{
+	GENERATED_USTRUCT_BODY()
+	DATAFLOW_NODE_DEFINE_INTERNAL(FComputeMuscleActivationDataNode, "ComputeMuscleActivationData", "Flesh", "")
+
+public:
+	UPROPERTY(meta = (DataflowInput, DataflowOutput, DisplayName = "Collection", DataflowPassthrough = "Collection"))
+	FManagedArrayCollection Collection;
+
+	UPROPERTY(meta = (DataflowInput, DisplayName = "OriginIndicesIn"))
+		TArray<int32> OriginIndicesIn;
+
+	UPROPERTY(meta = (DataflowInput, DisplayName = "InsertionIndicesIn"))
+		TArray<int32> InsertionIndicesIn;
+
+	FComputeMuscleActivationDataNode(const Dataflow::FNodeParameters& InParam, FGuid InGuid = FGuid::NewGuid())
+		: FDataflowNode(InParam, InGuid)
+	{
+		RegisterInputConnection(&Collection);
+		RegisterInputConnection(&OriginIndicesIn);
+		RegisterInputConnection(&InsertionIndicesIn);
+		RegisterOutputConnection(&Collection, &Collection);
+	}
+
+	virtual void Evaluate(Dataflow::FContext& Context, const FDataflowOutput* Out) const override;
+};
diff --git a/Engine/Plugins/Experimental/ChaosFlesh/Source/ChaosFleshNodes/Public/Dataflow/ChaosFleshEngineAssetNodes.h b/Engine/Plugins/Experimental/ChaosFlesh/Source/ChaosFleshNodes/Public/Dataflow/ChaosFleshEngineAssetNodes.h
deleted file mode 100644
diff --git a/Engine/Plugins/Experimental/ChaosFlesh/Source/ChaosFleshNodes/Public/Dataflow/ChaosFleshFleshAssetTerminalNode.h b/Engine/Plugins/Experimental/ChaosFlesh/Source/ChaosFleshNodes/Public/Dataflow/ChaosFleshFleshAssetTerminalNode.h
new file mode 100644
--- /dev/null
+++ b/Engine/Plugins/Experimental/ChaosFlesh/Source/ChaosFleshNodes/Public/Dataflow/ChaosFleshFleshAssetTerminalNode.h
@@ -0,0 +1,36 @@
+// Copyright Epic Games, Inc. All Rights Reserved.
+
+#pragma once 
+
+#include "CoreMinimal.h"
+#include "ChaosFlesh/FleshAsset.h"
+#include "Dataflow/DataflowCore.h"
+#include "Dataflow/DataflowEngine.h"
+
+#include "ChaosFleshFleshAssetTerminalNode.generated.h"
+
+USTRUCT(meta = (DataflowFlesh, DataflowTerminal))
+struct FFleshAssetTerminalDataflowNode : public FDataflowTerminalNode
+{
+	GENERATED_USTRUCT_BODY()
+	DATAFLOW_NODE_DEFINE_INTERNAL(FFleshAssetTerminalDataflowNode, "FleshAssetTerminal", "Terminal", "")
+	DATAFLOW_NODE_RENDER_TYPE(FGeometryCollection::StaticType(), "Collection")
+
+public:
+
+	UPROPERTY(meta = (DataflowInput, DataflowOutput, DisplayName = "Collection", DataflowPassthrough  = "Collection"))
+	FManagedArrayCollection Collection;
+
+
+	FFleshAssetTerminalDataflowNode(const Dataflow::FNodeParameters& InParam, FGuid InGuid = FGuid::NewGuid())
+	: FDataflowTerminalNode(InParam, InGuid)
+	{
+		RegisterInputConnection(&Collection);
+		RegisterOutputConnection(&Collection, &Collection);
+	}
+
+	virtual void SetAssetValue(TObjectPtr<UObject> Asset, Dataflow::FContext& Context) const override;
+	virtual void Evaluate(Dataflow::FContext& Context, const FDataflowOutput* Out) const override;
+};
+
+
diff --git a/Engine/Plugins/Experimental/ChaosFlesh/Source/ChaosFleshNodes/Public/Dataflow/ChaosFleshGenerateFiberDirectionsNode.h b/Engine/Plugins/Experimental/ChaosFlesh/Source/ChaosFleshNodes/Public/Dataflow/ChaosFleshGenerateFiberDirectionsNode.h
deleted file mode 100644
diff --git a/Engine/Plugins/Experimental/ChaosFlesh/Source/ChaosFleshNodes/Public/Dataflow/ChaosFleshGenerateOriginInsertionNode.h b/Engine/Plugins/Experimental/ChaosFlesh/Source/ChaosFleshNodes/Public/Dataflow/ChaosFleshGenerateOriginInsertionNode.h
new file mode 100644
--- /dev/null
+++ b/Engine/Plugins/Experimental/ChaosFlesh/Source/ChaosFleshNodes/Public/Dataflow/ChaosFleshGenerateOriginInsertionNode.h
@@ -0,0 +1,52 @@
+// Copyright Epic Games, Inc. All Rights Reserved.
+
+#pragma once 
+
+#include "CoreMinimal.h"
+#include "Dataflow/DataflowCore.h"
+#include "Dataflow/DataflowEngine.h"
+
+#include "ChaosFleshGenerateOriginInsertionNode.generated.h"
+
+//Given two sets of vertex indices, generate two sets of vertex indices for origins and insertions that are within X distance away.
+USTRUCT(meta = (DataflowFlesh))
+struct FGenerateOriginInsertionNode : public FDataflowNode
+{
+	GENERATED_USTRUCT_BODY()
+	DATAFLOW_NODE_DEFINE_INTERNAL(FGenerateOriginInsertionNode, "GenerateOriginInsertion", "Flesh", "")
+	DATAFLOW_NODE_RENDER_TYPE(FGeometryCollection::StaticType(), "Collection")
+
+public:
+	//typedef FManagedArrayCollection DataType;
+
+	UPROPERTY(meta = (DataflowInput, DataflowOutput, DisplayName = "Collection", DataflowPassthrough = "Collection"))
+	FManagedArrayCollection Collection;
+
+	UPROPERTY(meta = (DataflowInput, DisplayName = "OriginIndicesIn"))
+	TArray<int32> OriginIndicesIn;
+
+	UPROPERTY(meta = (DataflowInput, DisplayName = "InsertionIndicesIn"))
+	TArray<int32> InsertionIndicesIn;
+
+	UPROPERTY(meta = (DataflowOutput, DisplayName = "OriginIndicesOut"))
+	TArray<int32> OriginIndicesOut;
+
+	UPROPERTY(meta = (DataflowOutput, DisplayName = "InsertionIndicesOut"))
+	TArray<int32> InsertionIndicesOut;
+	
+	UPROPERTY(EditAnywhere, Category = "Dataflow")
+	float Radius = float(1);
+
+	FGenerateOriginInsertionNode(const Dataflow::FNodeParameters& InParam, FGuid InGuid = FGuid::NewGuid())
+		: FDataflowNode(InParam, InGuid)
+	{
+		RegisterInputConnection(&Collection);
+		RegisterInputConnection(&OriginIndicesIn);
+		RegisterInputConnection(&InsertionIndicesIn);
+		RegisterOutputConnection(&Collection, &Collection);
+		RegisterOutputConnection(&OriginIndicesOut);
+		RegisterOutputConnection(&InsertionIndicesOut);
+	}
+
+	virtual void Evaluate(Dataflow::FContext& Context, const FDataflowOutput* Out) const override;
+};
diff --git a/Engine/Plugins/Experimental/ChaosFlesh/Source/ChaosFleshNodes/Public/Dataflow/ChaosFleshGetFleshAssetNode.h b/Engine/Plugins/Experimental/ChaosFlesh/Source/ChaosFleshNodes/Public/Dataflow/ChaosFleshGetFleshAssetNode.h
new file mode 100644
--- /dev/null
+++ b/Engine/Plugins/Experimental/ChaosFlesh/Source/ChaosFleshNodes/Public/Dataflow/ChaosFleshGetFleshAssetNode.h
@@ -0,0 +1,34 @@
+// Copyright Epic Games, Inc. All Rights Reserved.
+
+#pragma once 
+
+#include "CoreMinimal.h"
+#include "ChaosFlesh/FleshAsset.h"
+#include "Dataflow/DataflowCore.h"
+#include "Dataflow/DataflowEngine.h"
+
+#include "ChaosFleshGetFleshAssetNode.generated.h"
+
+USTRUCT(meta = (DataflowFlesh))
+struct FGetFleshAssetDataflowNode : public FDataflowNode
+{
+	GENERATED_USTRUCT_BODY()
+	DATAFLOW_NODE_DEFINE_INTERNAL(FGetFleshAssetDataflowNode, "GetFleshAsset", "Flesh", "")
+	DATAFLOW_NODE_RENDER_TYPE(FGeometryCollection::StaticType(), "Collection")
+
+public:
+
+	UPROPERTY(EditAnywhere, Category = "Dataflow", meta = (DisplayName = "FleshAsset"))
+	TObjectPtr<const UFleshAsset> FleshAsset = nullptr;
+
+	UPROPERTY(meta = (DataflowOutput, DisplayName = "Collection"))
+	FManagedArrayCollection Output;
+
+	FGetFleshAssetDataflowNode(const Dataflow::FNodeParameters& InParam, FGuid InGuid = FGuid::NewGuid())
+		: FDataflowNode(InParam, InGuid)
+	{
+		RegisterOutputConnection(&Output);
+	}
+
+	virtual void Evaluate(Dataflow::FContext& Context, const FDataflowOutput* Out) const override;
+};
diff --git a/Engine/Plugins/Experimental/ChaosFlesh/Source/ChaosFleshNodes/Public/Dataflow/ChaosFleshGetSurfaceIndicesNode.h b/Engine/Plugins/Experimental/ChaosFlesh/Source/ChaosFleshNodes/Public/Dataflow/ChaosFleshGetSurfaceIndicesNode.h
new file mode 100644
--- /dev/null
+++ b/Engine/Plugins/Experimental/ChaosFlesh/Source/ChaosFleshNodes/Public/Dataflow/ChaosFleshGetSurfaceIndicesNode.h
@@ -0,0 +1,43 @@
+// Copyright Epic Games, Inc. All Rights Reserved.
+
+#pragma once 
+
+#include "CoreMinimal.h"
+#include "Dataflow/DataflowCore.h"
+#include "Dataflow/DataflowEngine.h"
+
+#include "ChaosFleshGetSurfaceIndicesNode.generated.h"
+
+USTRUCT(meta = (DataflowFlesh))
+struct FGetSurfaceIndicesNode : public FDataflowNode
+{
+	GENERATED_USTRUCT_BODY()
+	DATAFLOW_NODE_DEFINE_INTERNAL(FGetSurfaceIndicesNode, "GetSurfaceIndices", "Flesh", "")
+	DATAFLOW_NODE_RENDER_TYPE(FGeometryCollection::StaticType(), "Collection")
+
+public:
+	UPROPERTY(meta = (DataflowInput, DisplayName = "Collection"))
+	FManagedArrayCollection Collection;
+
+	UPROPERTY(meta = (DataflowInput, DisplayName = "GeometryGroupGuidsIn"))
+	TArray<FString> GeometryGroupGuidsIn;
+
+	UPROPERTY(meta = (DataflowOutput, DisplayName = "SurfaceIndicesOut"))
+	TArray<int32> SurfaceIndicesOut;
+
+	FGetSurfaceIndicesNode(const Dataflow::FNodeParameters& InParam, FGuid InGuid = FGuid::NewGuid())
+		: FDataflowNode(InParam, InGuid)
+	{
+		RegisterInputConnection(&Collection);
+		RegisterInputConnection(&GeometryGroupGuidsIn);
+		RegisterOutputConnection(&SurfaceIndicesOut);
+	}
+
+	virtual void Evaluate(Dataflow::FContext& Context, const FDataflowOutput* Out) const override;
+};
+
+namespace Dataflow
+{
+	void RegisterChaosFleshEngineAssetNodes();
+}
+
diff --git a/Engine/Plugins/Experimental/ChaosFlesh/Source/ChaosFleshNodes/Public/Dataflow/ChaosFleshImportGEO.h b/Engine/Plugins/Experimental/ChaosFlesh/Source/ChaosFleshNodes/Public/Dataflow/ChaosFleshImportGEO.h
deleted file mode 100644
diff --git a/Engine/Plugins/Experimental/ChaosFlesh/Source/ChaosFleshNodes/Public/Dataflow/ChaosFleshIsolateComponentNode.h b/Engine/Plugins/Experimental/ChaosFlesh/Source/ChaosFleshNodes/Public/Dataflow/ChaosFleshIsolateComponentNode.h
new file mode 100644
--- /dev/null
+++ b/Engine/Plugins/Experimental/ChaosFlesh/Source/ChaosFleshNodes/Public/Dataflow/ChaosFleshIsolateComponentNode.h
@@ -0,0 +1,39 @@
+// Copyright Epic Games, Inc. All Rights Reserved.
+
+#pragma once 
+
+#include "CoreMinimal.h"
+#include "Dataflow/DataflowCore.h"
+#include "Dataflow/DataflowEngine.h"
+
+#include "ChaosFleshIsolateComponentNode.generated.h"
+
+
+USTRUCT(meta = (DataflowFlesh))
+struct FIsolateComponentNode : public FDataflowNode
+{
+	GENERATED_USTRUCT_BODY()
+	DATAFLOW_NODE_DEFINE_INTERNAL(FIsolateComponentNode, "IsolateComponent", "Flesh", "")
+	DATAFLOW_NODE_RENDER_TYPE(FGeometryCollection::StaticType(), "Collection")
+
+public:
+	//typedef FManagedArrayCollection DataType;
+
+	UPROPERTY(meta = (DataflowInput, DataflowOutput, DisplayName = "Collection", DataflowPassthrough = "Collection"))
+	FManagedArrayCollection Collection;
+
+	UPROPERTY(EditAnywhere, Category = "Dataflow")
+	bool bDeleteHiddenFaces = false;
+
+	UPROPERTY(EditAnywhere, Category = "Dataflow")
+	FString TargetGeometryIndex = "";
+
+	FIsolateComponentNode(const Dataflow::FNodeParameters& InParam, FGuid InGuid = FGuid::NewGuid())
+		: FDataflowNode(InParam, InGuid)
+	{
+		RegisterInputConnection(&Collection);
+		RegisterOutputConnection(&Collection, &Collection);
+	}
+
+	virtual void Evaluate(Dataflow::FContext& Context, const FDataflowOutput* Out) const override;
+};
diff --git a/Engine/Plugins/Experimental/ChaosFlesh/Source/ChaosFleshNodes/Public/Dataflow/ChaosFleshPositionTargetInitializationNodes.h b/Engine/Plugins/Experimental/ChaosFlesh/Source/ChaosFleshNodes/Public/Dataflow/ChaosFleshPositionTargetInitializationNodes.h
deleted file mode 100644
diff --git a/Engine/Plugins/Experimental/ChaosFlesh/Source/ChaosFleshNodes/Public/Dataflow/ChaosFleshSetFleshBonePositionTargetBindingNode.h b/Engine/Plugins/Experimental/ChaosFlesh/Source/ChaosFleshNodes/Public/Dataflow/ChaosFleshSetFleshBonePositionTargetBindingNode.h
new file mode 100644
--- /dev/null
+++ b/Engine/Plugins/Experimental/ChaosFlesh/Source/ChaosFleshNodes/Public/Dataflow/ChaosFleshSetFleshBonePositionTargetBindingNode.h
@@ -0,0 +1,63 @@
+// Copyright Epic Games, Inc. All Rights Reserved.
+
+#pragma once 
+
+#include "CoreMinimal.h"
+#include "Dataflow/DataflowCore.h"
+#include "Dataflow/DataflowEngine.h"
+
+#include "ChaosFleshSetFleshBonePositionTargetBindingNode.generated.h"
+
+class USkeletalMesh;
+
+UENUM(BlueprintType)
+enum class ESkeletalBindingMode : uint8
+{
+	Dataflow_SkeletalBinding_Kinematic UMETA(DisplayName = "Kinematic"),
+	Dataflow_SkeletalBinding_PositionTarget UMETA(DisplayName = "Position Target"),
+	//
+	Chaos_Max UMETA(Hidden)
+};
+
+USTRUCT(meta = (DataflowFlesh))
+struct FSetFleshBonePositionTargetBindingDataflowNode : public FDataflowNode
+{
+	GENERATED_USTRUCT_BODY()
+		DATAFLOW_NODE_DEFINE_INTERNAL(FSetFleshBonePositionTargetBindingDataflowNode, "FSetFleshBonePositionTargetBinding", "Flesh", "")
+		DATAFLOW_NODE_RENDER_TYPE(FGeometryCollection::StaticType(), "Collection")
+
+public:
+	typedef FManagedArrayCollection DataType;
+
+
+	UPROPERTY(meta = (DataflowInput, DataflowOutput, DisplayName = "Collection", DataflowPassthrough = "Collection"))
+		FManagedArrayCollection Collection;
+
+	UPROPERTY(EditAnywhere, Category = "Dataflow")
+		ESkeletalBindingMode SkeletalBindingMode = ESkeletalBindingMode::Dataflow_SkeletalBinding_PositionTarget;
+
+	UPROPERTY(EditAnywhere, Category = "Dataflow")
+		float PositionTargetStiffness = 10000.f;
+
+	UPROPERTY(meta = (DataflowInput, DisplayName = "SkeletalMesh"))
+		TObjectPtr<USkeletalMesh> SkeletalMeshIn = nullptr;
+
+	UPROPERTY(EditAnywhere, Category = "Dataflow", meta = (ClampMin = "0.0", ClampMax = "2.0", UIMin = "0.0", UIMax = "2.0"))
+		float VertexRadiusRatio = .001f;
+
+
+	FSetFleshBonePositionTargetBindingDataflowNode(const Dataflow::FNodeParameters& InParam, FGuid InGuid = FGuid::NewGuid())
+		: FDataflowNode(InParam, InGuid)
+	{
+		RegisterInputConnection(&SkeletalMeshIn);
+		RegisterInputConnection(&Collection);
+		RegisterOutputConnection(&Collection, &Collection);
+	}
+
+	virtual void Evaluate(Dataflow::FContext& Context, const FDataflowOutput* Out) const override;
+};
+
+namespace Dataflow
+{
+	void RegisterChaosFleshPositionTargetInitializationNodes();
+}
diff --git a/Engine/Plugins/Experimental/ChaosFlesh/Source/ChaosFleshNodes/Public/Dataflow/ChaosFleshSetVertexTetrahedraPositionTargetBindingNode.h b/Engine/Plugins/Experimental/ChaosFlesh/Source/ChaosFleshNodes/Public/Dataflow/ChaosFleshSetVertexTetrahedraPositionTargetBindingNode.h
new file mode 100644
--- /dev/null
+++ b/Engine/Plugins/Experimental/ChaosFlesh/Source/ChaosFleshNodes/Public/Dataflow/ChaosFleshSetVertexTetrahedraPositionTargetBindingNode.h
@@ -0,0 +1,46 @@
+// Copyright Epic Games, Inc. All Rights Reserved.
+
+#pragma once 
+
+#include "CoreMinimal.h"
+#include "Dataflow/DataflowCore.h"
+#include "Dataflow/DataflowEngine.h"
+
+#include "ChaosFleshSetVertexTetrahedraPositionTargetBindingNode.generated.h"
+
+USTRUCT(meta = (DataflowFlesh))
+struct FSetVertexTetrahedraPositionTargetBindingDataflowNode : public FDataflowNode
+{
+	GENERATED_USTRUCT_BODY()
+		DATAFLOW_NODE_DEFINE_INTERNAL(FSetVertexTetrahedraPositionTargetBindingDataflowNode, "FSetVertexTetrahedraPositionTargetBinding", "Flesh", "")
+		DATAFLOW_NODE_RENDER_TYPE(FGeometryCollection::StaticType(), "Collection")
+
+public:
+	typedef FManagedArrayCollection DataType;
+
+
+	UPROPERTY(meta = (DataflowInput, DataflowOutput, DisplayName = "Collection", DataflowPassthrough = "Collection"))
+		FManagedArrayCollection Collection;
+
+	UPROPERTY(meta = (DataflowInput, DisplayName = "TargetIndicesIn"))
+		TArray<int32> TargetIndicesIn;
+
+	UPROPERTY(meta = (DataflowInput, DisplayName = "GeometryGroupGuidsIn"))
+		TArray<FString> GeometryGroupGuidsIn;
+
+	UPROPERTY(EditAnywhere, Category = "Dataflow")
+		float PositionTargetStiffness = 10000.f;
+
+
+
+	FSetVertexTetrahedraPositionTargetBindingDataflowNode(const Dataflow::FNodeParameters& InParam, FGuid InGuid = FGuid::NewGuid())
+		: FDataflowNode(InParam, InGuid)
+	{
+		RegisterInputConnection(&Collection);
+		RegisterOutputConnection(&Collection, &Collection);
+		RegisterInputConnection(&TargetIndicesIn);
+		RegisterInputConnection(&GeometryGroupGuidsIn);
+	}
+
+	virtual void Evaluate(Dataflow::FContext& Context, const FDataflowOutput* Out) const override;
+};
diff --git a/Engine/Plugins/Experimental/ChaosFlesh/Source/ChaosFleshNodes/Public/Dataflow/ChaosFleshSetVertexTrianglePositionTargetBindingNode.h b/Engine/Plugins/Experimental/ChaosFlesh/Source/ChaosFleshNodes/Public/Dataflow/ChaosFleshSetVertexTrianglePositionTargetBindingNode.h
new file mode 100644
--- /dev/null
+++ b/Engine/Plugins/Experimental/ChaosFlesh/Source/ChaosFleshNodes/Public/Dataflow/ChaosFleshSetVertexTrianglePositionTargetBindingNode.h
@@ -0,0 +1,43 @@
+// Copyright Epic Games, Inc. All Rights Reserved.
+
+#pragma once 
+
+#include "CoreMinimal.h"
+#include "Dataflow/DataflowCore.h"
+#include "Dataflow/DataflowEngine.h"
+
+#include "ChaosFleshSetVertexTrianglePositionTargetBindingNode.generated.h"
+
+USTRUCT(meta = (DataflowFlesh))
+struct FSetVertexTrianglePositionTargetBindingDataflowNode : public FDataflowNode
+{
+	GENERATED_USTRUCT_BODY()
+		DATAFLOW_NODE_DEFINE_INTERNAL(FSetVertexTrianglePositionTargetBindingDataflowNode, "FSetVertexTrianglePositionTargetBinding", "Flesh", "")
+		DATAFLOW_NODE_RENDER_TYPE(FGeometryCollection::StaticType(), "Collection")
+
+public:
+	typedef FManagedArrayCollection DataType;
+
+
+	UPROPERTY(meta = (DataflowInput, DataflowOutput, DisplayName = "Collection", DataflowPassthrough = "Collection"))
+		FManagedArrayCollection Collection;
+
+	UPROPERTY(EditAnywhere, Category = "Dataflow")
+		float PositionTargetStiffness = 10000.f;
+
+	//UPROPERTY(EditAnywhere, Category = "Dataflow", meta = (ClampMin = "0.0", ClampMax = "1.0", UIMin = "0.0", UIMax = "1.0"))
+	//	float TriangleRadiusPadding = .1f;
+
+	UPROPERTY(EditAnywhere, Category = "Dataflow", meta = (ClampMin = "0.0", ClampMax = "2.0", UIMin = "0.0", UIMax = "2.0"))
+		float VertexRadiusRatio = .001f;
+
+
+	FSetVertexTrianglePositionTargetBindingDataflowNode(const Dataflow::FNodeParameters& InParam, FGuid InGuid = FGuid::NewGuid())
+		: FDataflowNode(InParam, InGuid)
+	{
+		RegisterInputConnection(&Collection);
+		RegisterOutputConnection(&Collection, &Collection);
+	}
+
+	virtual void Evaluate(Dataflow::FContext& Context, const FDataflowOutput* Out) const override;
+};
diff --git a/Engine/Plugins/Experimental/ChaosFlesh/Source/ChaosFleshNodes/Public/Dataflow/ChaosFleshSetVertexVertexPositionTargetBindingNode.h b/Engine/Plugins/Experimental/ChaosFlesh/Source/ChaosFleshNodes/Public/Dataflow/ChaosFleshSetVertexVertexPositionTargetBindingNode.h
new file mode 100644
--- /dev/null
+++ b/Engine/Plugins/Experimental/ChaosFlesh/Source/ChaosFleshNodes/Public/Dataflow/ChaosFleshSetVertexVertexPositionTargetBindingNode.h
@@ -0,0 +1,45 @@
+// Copyright Epic Games, Inc. All Rights Reserved.
+
+#pragma once 
+
+#include "CoreMinimal.h"
+#include "Dataflow/DataflowCore.h"
+#include "Dataflow/DataflowEngine.h"
+
+#include "ChaosFleshSetVertexVertexPositionTargetBindingNode.generated.h"
+
+
+USTRUCT(meta = (DataflowFlesh))
+struct FSetVertexVertexPositionTargetBindingDataflowNode : public FDataflowNode
+{
+	GENERATED_USTRUCT_BODY()
+	DATAFLOW_NODE_DEFINE_INTERNAL(FSetVertexVertexPositionTargetBindingDataflowNode, "SetVertexVertexPositionTargetBinding", "Flesh", "")
+	DATAFLOW_NODE_RENDER_TYPE(FGeometryCollection::StaticType(), "Collection")
+
+public:
+	typedef FManagedArrayCollection DataType;
+
+	UPROPERTY(EditAnywhere, Category = "Dataflow", meta = (ClampMin = "0.0", ClampMax = "1.0", UIMin = "0.0", UIMax = "1.0"))
+		float RadiusRatio = .1f;
+
+	UPROPERTY(meta = (DataflowInput, DataflowOutput, DisplayName = "Collection", DataflowPassthrough = "Collection"))
+		FManagedArrayCollection Collection;
+
+	UPROPERTY(meta = (DataflowInput, DisplayName = "TargetIndicesIn"))
+		TArray<int32> TargetIndicesIn;
+
+	UPROPERTY(EditAnywhere, Category = "Dataflow")
+		float PositionTargetStiffness = 10000.f;
+
+
+
+	FSetVertexVertexPositionTargetBindingDataflowNode(const Dataflow::FNodeParameters& InParam, FGuid InGuid = FGuid::NewGuid())
+		: FDataflowNode(InParam, InGuid)
+	{
+		RegisterInputConnection(&Collection);
+		RegisterOutputConnection(&Collection, &Collection);
+		RegisterInputConnection(&TargetIndicesIn);
+	}
+
+	virtual void Evaluate(Dataflow::FContext& Context, const FDataflowOutput* Out) const override;
+};
diff --git a/Engine/Plugins/Experimental/ChaosFlesh/Source/ChaosFleshNodes/Public/Dataflow/ChaosFleshVisualizeFiberFieldNode.h b/Engine/Plugins/Experimental/ChaosFlesh/Source/ChaosFleshNodes/Public/Dataflow/ChaosFleshVisualizeFiberFieldNode.h
new file mode 100644
--- /dev/null
+++ b/Engine/Plugins/Experimental/ChaosFlesh/Source/ChaosFleshNodes/Public/Dataflow/ChaosFleshVisualizeFiberFieldNode.h
@@ -0,0 +1,40 @@
+// Copyright Epic Games, Inc. All Rights Reserved.
+
+#pragma once 
+
+#include "CoreMinimal.h"
+#include "Dataflow/DataflowCore.h"
+#include "Dataflow/DataflowEngine.h"
+#include "Field/FieldSystemTypes.h"
+
+#include "ChaosFleshVisualizeFiberFieldNode.generated.h"
+
+/**
+* Visualizes a muscle fiber direction per tetrahedron from a GeometryCollection containing tetrahedra.
+*/
+USTRUCT(meta = (DataflowFlesh))
+struct FVisualizeFiberFieldNode : public FDataflowNode
+{
+	GENERATED_USTRUCT_BODY()
+	DATAFLOW_NODE_DEFINE_INTERNAL(FVisualizeFiberFieldNode, "VisualizeFiberField", "Flesh", "")
+	DATAFLOW_NODE_RENDER_TYPE(FFieldCollection::StaticType(), "VectorField")
+
+public:
+	UPROPERTY(meta = (DataflowInput, DisplayName = "Collection"))
+	FManagedArrayCollection Collection;
+
+	UPROPERTY(EditAnywhere, Category = "Dataflow")
+	float VectorScale = 1.0;
+
+	UPROPERTY(meta = (DataflowOutput, DisplayName = "VectorField"))
+	FFieldCollection VectorField;
+	
+	FVisualizeFiberFieldNode(const Dataflow::FNodeParameters& InParam, FGuid InGuid = FGuid::NewGuid())
+		: FDataflowNode(InParam, InGuid)
+	{
+		RegisterInputConnection(&Collection);
+		RegisterOutputConnection(&VectorField);
+	}
+
+	virtual void Evaluate(Dataflow::FContext& Context, const FDataflowOutput* Out) const override;
+};
