ChaosFlesh : Renamed OriginsInsertionsNode

diff --git a/Engine/Plugins/Experimental/ChaosFlesh/Source/ChaosFleshDeprecatedNodes/Private/Dataflow/ChaosFleshDeprecatedNodesPlugin.cpp b/Engine/Plugins/Experimental/ChaosFlesh/Source/ChaosFleshDeprecatedNodes/Private/Dataflow/ChaosFleshDeprecatedNodesPlugin.cpp
--- a/Engine/Plugins/Experimental/ChaosFlesh/Source/ChaosFleshDeprecatedNodes/Private/Dataflow/ChaosFleshDeprecatedNodesPlugin.cpp
+++ b/Engine/Plugins/Experimental/ChaosFlesh/Source/ChaosFleshDeprecatedNodes/Private/Dataflow/ChaosFleshDeprecatedNodesPlugin.cpp
@@ -5,6 +5,7 @@
 #include "Dataflow/ChaosFleshCreateTetrahedralCollectionNode.h"
 #include "Dataflow/ChaosFleshConstructTetGridNode.h"
 #include "Dataflow/ChaosFleshImportGEO.h"
+#include "Dataflow/ChaosFleshKinematicOriginInsertionInitializationNode.h"
 
 
 #define LOCTEXT_NAMESPACE "ChaosFleshDeprecatedNodes"
@@ -15,6 +16,7 @@ void IChaosFleshDeprecatedNodesPlugin::StartupModule()
 	Dataflow::RegisterChaosFleshImportGEONodes();
 	DATAFLOW_NODE_REGISTER_CREATION_FACTORY(FGenerateTetrahedralCollectionDataflowNodes);
 	DATAFLOW_NODE_REGISTER_CREATION_FACTORY(FConstructTetGridNode);
+	DATAFLOW_NODE_REGISTER_CREATION_FACTORY(FKinematicOriginInsertionInitializationDataflowNode);
 }
 
 void IChaosFleshDeprecatedNodesPlugin::ShutdownModule()
diff --git a/Engine/Plugins/Experimental/ChaosFlesh/Source/ChaosFleshDeprecatedNodes/Private/Dataflow/ChaosFleshKinematicOriginInsertionInitializationNode.cpp b/Engine/Plugins/Experimental/ChaosFlesh/Source/ChaosFleshDeprecatedNodes/Private/Dataflow/ChaosFleshKinematicOriginInsertionInitializationNode.cpp
new file mode 100644
--- /dev/null
+++ b/Engine/Plugins/Experimental/ChaosFlesh/Source/ChaosFleshDeprecatedNodes/Private/Dataflow/ChaosFleshKinematicOriginInsertionInitializationNode.cpp
@@ -0,0 +1,137 @@
+// Copyright Epic Games, Inc. All Rights Reserved.
+
+#include "Dataflow/ChaosFleshKinematicOriginInsertionInitializationNode.h"
+
+#include "Engine/SkeletalMesh.h"
+#include "GeometryCollection/Facades/CollectionKinematicBindingFacade.h"
+#include "GeometryCollection/Facades/CollectionVertexBoneWeightsFacade.h"
+#include "Rendering/SkeletalMeshRenderData.h"
+#include "BoneWeights.h"
+
+#include UE_INLINE_GENERATED_CPP_BY_NAME(ChaosFleshKinematicOriginInsertionInitializationNode)
+
+
+void FKinematicOriginInsertionInitializationDataflowNode::Evaluate(Dataflow::FContext& Context, const FDataflowOutput* Out) const
+{
+	if (Out->IsA<DataType>(&Collection))
+	{
+		DataType InCollection = GetValue<DataType>(Context, &Collection);
+
+		if (TManagedArray<FVector3f>* Vertices = InCollection.FindAttribute<FVector3f>("Vertex", FGeometryCollection::VerticesGroup))
+		{
+			if (FindInput(&OriginVertexIndicesIn) && FindInput(&OriginVertexIndicesIn)->GetConnection() && FindInput(&InsertionVertexIndicesIn) && FindInput(&InsertionVertexIndicesIn)->GetConnection())
+			{
+				TArray<int32> BoundVerts;
+				TArray<float> BoundWeights;
+
+				for (int32 SelectionIndex : GetValue<TArray<int32>>(Context, &OriginVertexIndicesIn))
+				{
+					if (0 <= SelectionIndex && SelectionIndex < Vertices->Num())
+					{
+						BoundVerts.Add(SelectionIndex);
+					}
+				}
+				for (int32 SelectionIndex : GetValue<TArray<int32>>(Context, &InsertionVertexIndicesIn))
+				{
+					if (0 <= SelectionIndex && SelectionIndex < Vertices->Num())
+					{
+						BoundVerts.Add(SelectionIndex);
+					}
+				}
+				if (TObjectPtr<USkeletalMesh> BoneSkeletalMesh = GetValue<TObjectPtr<USkeletalMesh>>(Context, &BoneSkeletalMeshIn))
+				{
+					FSkeletalMeshRenderData* RenderData = BoneSkeletalMesh->GetResourceForRendering();
+					if (RenderData->LODRenderData.Num())
+					{
+						//Grab vertices only, no elements
+						FSkeletalMeshLODRenderData* LODRenderData = &RenderData->LODRenderData[0];
+						const FPositionVertexBuffer& PositionVertexBuffer =
+							LODRenderData->StaticVertexBuffers.PositionVertexBuffer;
+						//Grab skin weights
+						const FSkinWeightVertexBuffer* SkinWeightVertexBuffer = LODRenderData->GetSkinWeightVertexBuffer();
+						const int32 MaxBoneInfluences = SkinWeightVertexBuffer->GetMaxBoneInfluences();
+						TArray<FTransform> ComponentPose;
+						Dataflow::Animation::GlobalTransforms(BoneSkeletalMesh->GetRefSkeleton(), ComponentPose);
+						TArray<TArray<int32>> BoneBoundVerts;
+						TArray<TArray<float>> BoneBoundWeights;
+						BoneBoundVerts.SetNum(ComponentPose.Num());
+						BoneBoundWeights.SetNum(ComponentPose.Num());
+						if (!PositionVertexBuffer.GetNumVertices())
+						{
+							return;
+						}
+						auto DoubleVert = [](FVector3f V) { return FVector3d(V.X, V.Y, V.Z); };
+						for (int32 i = 0; i < BoundVerts.Num(); i++)
+						{
+							int ClosestPointIndex = 0;
+							double MinDistance = DBL_MAX;
+							for (uint32 j = 0; j < PositionVertexBuffer.GetNumVertices(); j++)
+							{
+								const FVector3f& Pos = PositionVertexBuffer.VertexPosition(j);
+								double Distance = FVector::Distance(DoubleVert((*Vertices)[BoundVerts[i]]), DoubleVert(Pos));
+								if (Distance < MinDistance)
+								{
+									ClosestPointIndex = j;
+									MinDistance = Distance;
+								}
+							}
+							int32 SectionIndex;
+							int32 VertIndex;
+							LODRenderData->GetSectionFromVertexIndex(ClosestPointIndex, SectionIndex, VertIndex);
+
+							check(SectionIndex < LODRenderData->RenderSections.Num());
+							const FSkelMeshRenderSection& Section = LODRenderData->RenderSections[SectionIndex];
+							int32 BufferVertIndex = Section.GetVertexBufferIndex() + VertIndex;
+							for (int32 InfluenceIndex = 0; InfluenceIndex < MaxBoneInfluences; InfluenceIndex++)
+							{
+								const int32 BoneIndex = Section.BoneMap[SkinWeightVertexBuffer->GetBoneIndex(BufferVertIndex, InfluenceIndex)];
+								const float	Weight = (float)SkinWeightVertexBuffer->GetBoneWeight(BufferVertIndex, InfluenceIndex) * UE::AnimationCore::InvMaxRawBoneWeightFloat;
+								if (Weight > float(0) && 0 <= BoneIndex && BoneIndex < ComponentPose.Num())
+								{
+									FString BoneName = BoneSkeletalMesh->GetRefSkeleton().GetBoneName(BoneIndex).ToString();
+									BoneBoundVerts[BoneIndex].Add(BoundVerts[i]);
+									BoneBoundWeights[BoneIndex].Add(Weight);
+								}
+							}
+						}
+						for (int32 BoneIndex = 0; BoneIndex < ComponentPose.Num(); ++BoneIndex)
+						{
+							if (BoneBoundVerts[BoneIndex].Num())
+							{
+								FString BoneName = BoneSkeletalMesh->GetRefSkeleton().GetBoneName(BoneIndex).ToString();
+								//get local coords of bound verts
+								typedef GeometryCollection::Facades::FKinematicBindingFacade FKinematics;
+								FKinematics Kinematics(InCollection); Kinematics.DefineSchema();
+								if (Kinematics.IsValid())
+								{
+									FKinematics::FBindingKey Binding = Kinematics.SetBoneBindings(BoneIndex, BoneBoundVerts[BoneIndex], BoneBoundWeights[BoneIndex]);
+									TManagedArray<TArray<FVector3f>>& LocalPos = InCollection.AddAttribute<TArray<FVector3f>>("LocalPosition", Binding.GroupName);
+									Kinematics.AddKinematicBinding(Binding);
+
+									auto FloatVert = [](FVector3d V) { return FVector3f(V.X, V.Y, V.Z); };
+									LocalPos[Binding.Index].SetNum(BoneBoundVerts[BoneIndex].Num());
+									for (int32 i = 0; i < BoneBoundVerts[BoneIndex].Num(); i++)
+									{
+										FVector3f Temp = (*Vertices)[BoneBoundVerts[BoneIndex][i]];
+										LocalPos[Binding.Index][i] = FloatVert(ComponentPose[BoneIndex].InverseTransformPosition(DoubleVert(Temp)));
+									}
+								}
+							}
+						}
+						GeometryCollection::Facades::FVertexBoneWeightsFacade(InCollection).AddBoneWeightsFromKinematicBindings();
+					}
+				}
+				else
+				{
+					if (BoundVerts.Num())
+					{
+						BoundWeights.Init(1.0, BoundVerts.Num());
+						GeometryCollection::Facades::FKinematicBindingFacade Kinematics(InCollection);
+						Kinematics.AddKinematicBinding(Kinematics.SetBoneBindings(INDEX_NONE, BoundVerts, BoundWeights));
+					}
+				}
+			}
+		}
+		SetValue(Context, MoveTemp(InCollection), &Collection);
+	}
+}
diff --git a/Engine/Plugins/Experimental/ChaosFlesh/Source/ChaosFleshDeprecatedNodes/Public/Dataflow/ChaosFleshKinematicOriginInsertionInitializationNode.h b/Engine/Plugins/Experimental/ChaosFlesh/Source/ChaosFleshDeprecatedNodes/Public/Dataflow/ChaosFleshKinematicOriginInsertionInitializationNode.h
new file mode 100644
--- /dev/null
+++ b/Engine/Plugins/Experimental/ChaosFlesh/Source/ChaosFleshDeprecatedNodes/Public/Dataflow/ChaosFleshKinematicOriginInsertionInitializationNode.h
@@ -0,0 +1,48 @@
+// Copyright Epic Games, Inc. All Rights Reserved.
+#pragma once 
+
+#include "CoreMinimal.h"
+#include "Dataflow/DataflowCore.h"
+#include "Dataflow/DataflowEngine.h"
+
+#include "ChaosFleshKinematicOriginInsertionInitializationNode.generated.h"
+
+class USkeletalMesh;
+
+// @todo(deprecate), rename to FKinematicSomethingBetterThanThisNameDataflowNode
+USTRUCT(meta = (DataflowFlesh))
+struct FKinematicOriginInsertionInitializationDataflowNode : public FDataflowNode
+{
+	GENERATED_USTRUCT_BODY()
+	DATAFLOW_NODE_DEFINE_INTERNAL(FKinematicOriginInsertionInitializationDataflowNode, "KinematicOriginInsertionInitialization", "Flesh", "")
+	DATAFLOW_NODE_RENDER_TYPE(FGeometryCollection::StaticType(), "Collection")
+
+public:
+	typedef FManagedArrayCollection DataType;
+
+	UPROPERTY(meta = (DataflowInput, DataflowOutput, DisplayName = "Collection", DataflowPassthrough = "Collection"))
+		FManagedArrayCollection Collection;
+
+	UPROPERTY(meta = (DataflowInput, DisplayName = "OriginSelectionSet"))
+		TArray<int32> OriginVertexIndicesIn;
+
+	UPROPERTY(meta = (DataflowInput, DisplayName = "InsertionSelectionSet"))
+		TArray<int32> InsertionVertexIndicesIn;
+
+	UPROPERTY(meta = (DataflowInput, DisplayName = "BoneSkeletalMesh"))
+		TObjectPtr<USkeletalMesh> BoneSkeletalMeshIn;
+
+
+	FKinematicOriginInsertionInitializationDataflowNode(const Dataflow::FNodeParameters& InParam, FGuid InGuid = FGuid::NewGuid())
+		: FDataflowNode(InParam, InGuid)
+	{
+		RegisterInputConnection(&Collection);
+		RegisterOutputConnection(&Collection, &Collection);
+		RegisterInputConnection(&OriginVertexIndicesIn);
+		RegisterInputConnection(&InsertionVertexIndicesIn);
+		RegisterInputConnection(&BoneSkeletalMeshIn);
+	}
+
+	virtual void Evaluate(Dataflow::FContext& Context, const FDataflowOutput* Out) const override;
+};
+
diff --git a/Engine/Plugins/Experimental/ChaosFlesh/Source/ChaosFleshNodes/Private/Dataflow/ChaosFleshKinematicMuscleAttachmentsNode.cpp b/Engine/Plugins/Experimental/ChaosFlesh/Source/ChaosFleshNodes/Private/Dataflow/ChaosFleshKinematicMuscleAttachmentsNode.cpp
new file mode 100644
--- /dev/null
+++ b/Engine/Plugins/Experimental/ChaosFlesh/Source/ChaosFleshNodes/Private/Dataflow/ChaosFleshKinematicMuscleAttachmentsNode.cpp
@@ -0,0 +1,137 @@
+// Copyright Epic Games, Inc. All Rights Reserved.
+
+#include "Dataflow/ChaosFleshKinematicMuscleAttachmentsNode.h"
+
+#include "Engine/SkeletalMesh.h"
+#include "GeometryCollection/Facades/CollectionKinematicBindingFacade.h"
+#include "GeometryCollection/Facades/CollectionVertexBoneWeightsFacade.h"
+#include "Rendering/SkeletalMeshRenderData.h"
+#include "BoneWeights.h"
+
+#include UE_INLINE_GENERATED_CPP_BY_NAME(ChaosFleshKinematicMuscleAttachmentsNode)
+
+
+void FKinematicMuscleAttachmentsDataflowNode::Evaluate(Dataflow::FContext& Context, const FDataflowOutput* Out) const
+{
+	if (Out->IsA<DataType>(&Collection))
+	{
+		DataType InCollection = GetValue<DataType>(Context, &Collection);
+
+		if (TManagedArray<FVector3f>* Vertices = InCollection.FindAttribute<FVector3f>("Vertex", FGeometryCollection::VerticesGroup))
+		{
+			if (FindInput(&OriginVertexIndicesIn) && FindInput(&OriginVertexIndicesIn)->GetConnection() && FindInput(&InsertionVertexIndicesIn) && FindInput(&InsertionVertexIndicesIn)->GetConnection())
+			{
+				TArray<int32> BoundVerts;
+				TArray<float> BoundWeights;
+
+				for (int32 SelectionIndex : GetValue<TArray<int32>>(Context, &OriginVertexIndicesIn))
+				{
+					if (0 <= SelectionIndex && SelectionIndex < Vertices->Num())
+					{
+						BoundVerts.Add(SelectionIndex);
+					}
+				}
+				for (int32 SelectionIndex : GetValue<TArray<int32>>(Context, &InsertionVertexIndicesIn))
+				{
+					if (0 <= SelectionIndex && SelectionIndex < Vertices->Num())
+					{
+						BoundVerts.Add(SelectionIndex);
+					}
+				}
+				if (TObjectPtr<USkeletalMesh> SkeletalMesh = GetValue<TObjectPtr<USkeletalMesh>>(Context, &SkeletalMeshIn))
+				{
+					FSkeletalMeshRenderData* RenderData = SkeletalMesh->GetResourceForRendering();
+					if (RenderData->LODRenderData.Num())
+					{
+						//Grab vertices only, no elements
+						FSkeletalMeshLODRenderData* LODRenderData = &RenderData->LODRenderData[0];
+						const FPositionVertexBuffer& PositionVertexBuffer =
+							LODRenderData->StaticVertexBuffers.PositionVertexBuffer;
+						//Grab skin weights
+						const FSkinWeightVertexBuffer* SkinWeightVertexBuffer = LODRenderData->GetSkinWeightVertexBuffer();
+						const int32 MaxBoneInfluences = SkinWeightVertexBuffer->GetMaxBoneInfluences();
+						TArray<FTransform> ComponentPose;
+						Dataflow::Animation::GlobalTransforms(SkeletalMesh->GetRefSkeleton(), ComponentPose);
+						TArray<TArray<int32>> BoneBoundVerts;
+						TArray<TArray<float>> BoneBoundWeights;
+						BoneBoundVerts.SetNum(ComponentPose.Num());
+						BoneBoundWeights.SetNum(ComponentPose.Num());
+						if (!PositionVertexBuffer.GetNumVertices())
+						{
+							return;
+						}
+						auto DoubleVert = [](FVector3f V) { return FVector3d(V.X, V.Y, V.Z); };
+						for (int32 i = 0; i < BoundVerts.Num(); i++)
+						{
+							int ClosestPointIndex = 0;
+							double MinDistance = DBL_MAX;
+							for (uint32 j = 0; j < PositionVertexBuffer.GetNumVertices(); j++)
+							{
+								const FVector3f& Pos = PositionVertexBuffer.VertexPosition(j);
+								double Distance = FVector::Distance(DoubleVert((*Vertices)[BoundVerts[i]]), DoubleVert(Pos));
+								if (Distance < MinDistance)
+								{
+									ClosestPointIndex = j;
+									MinDistance = Distance;
+								}
+							}
+							int32 SectionIndex;
+							int32 VertIndex;
+							LODRenderData->GetSectionFromVertexIndex(ClosestPointIndex, SectionIndex, VertIndex);
+
+							check(SectionIndex < LODRenderData->RenderSections.Num());
+							const FSkelMeshRenderSection& Section = LODRenderData->RenderSections[SectionIndex];
+							int32 BufferVertIndex = Section.GetVertexBufferIndex() + VertIndex;
+							for (int32 InfluenceIndex = 0; InfluenceIndex < MaxBoneInfluences; InfluenceIndex++)
+							{
+								const int32 BoneIndex = Section.BoneMap[SkinWeightVertexBuffer->GetBoneIndex(BufferVertIndex, InfluenceIndex)];
+								const float	Weight = (float)SkinWeightVertexBuffer->GetBoneWeight(BufferVertIndex, InfluenceIndex) * UE::AnimationCore::InvMaxRawBoneWeightFloat;
+								if (Weight > float(0) && 0 <= BoneIndex && BoneIndex < ComponentPose.Num())
+								{
+									FString BoneName = SkeletalMesh->GetRefSkeleton().GetBoneName(BoneIndex).ToString();
+									BoneBoundVerts[BoneIndex].Add(BoundVerts[i]);
+									BoneBoundWeights[BoneIndex].Add(Weight);
+								}
+							}
+						}
+						for (int32 BoneIndex = 0; BoneIndex < ComponentPose.Num(); ++BoneIndex)
+						{
+							if (BoneBoundVerts[BoneIndex].Num())
+							{
+								FString BoneName = SkeletalMesh->GetRefSkeleton().GetBoneName(BoneIndex).ToString();
+								//get local coords of bound verts
+								typedef GeometryCollection::Facades::FKinematicBindingFacade FKinematics;
+								FKinematics Kinematics(InCollection); Kinematics.DefineSchema();
+								if (Kinematics.IsValid())
+								{
+									FKinematics::FBindingKey Binding = Kinematics.SetBoneBindings(BoneIndex, BoneBoundVerts[BoneIndex], BoneBoundWeights[BoneIndex]);
+									TManagedArray<TArray<FVector3f>>& LocalPos = InCollection.AddAttribute<TArray<FVector3f>>("LocalPosition", Binding.GroupName);
+									Kinematics.AddKinematicBinding(Binding);
+
+									auto FloatVert = [](FVector3d V) { return FVector3f(V.X, V.Y, V.Z); };
+									LocalPos[Binding.Index].SetNum(BoneBoundVerts[BoneIndex].Num());
+									for (int32 i = 0; i < BoneBoundVerts[BoneIndex].Num(); i++)
+									{
+										FVector3f Temp = (*Vertices)[BoneBoundVerts[BoneIndex][i]];
+										LocalPos[Binding.Index][i] = FloatVert(ComponentPose[BoneIndex].InverseTransformPosition(DoubleVert(Temp)));
+									}
+								}
+							}
+						}
+						GeometryCollection::Facades::FVertexBoneWeightsFacade(InCollection).AddBoneWeightsFromKinematicBindings();
+					}
+				}
+				else
+				{
+					if (BoundVerts.Num())
+					{
+						BoundWeights.Init(1.0, BoundVerts.Num());
+						GeometryCollection::Facades::FKinematicBindingFacade Kinematics(InCollection);
+						Kinematics.AddKinematicBinding(Kinematics.SetBoneBindings(INDEX_NONE, BoundVerts, BoundWeights));
+					}
+				}
+			}
+		}
+		SetValue(Context, MoveTemp(InCollection), &Collection);
+	}
+}
diff --git a/Engine/Plugins/Experimental/ChaosFlesh/Source/ChaosFleshNodes/Private/Dataflow/ChaosFleshKinematicOriginInsertionInitializationNode.cpp b/Engine/Plugins/Experimental/ChaosFlesh/Source/ChaosFleshNodes/Private/Dataflow/ChaosFleshKinematicOriginInsertionInitializationNode.cpp
deleted file mode 100644
diff --git a/Engine/Plugins/Experimental/ChaosFlesh/Source/ChaosFleshNodes/Private/Dataflow/ChaosFleshNodesPlugin.cpp b/Engine/Plugins/Experimental/ChaosFlesh/Source/ChaosFleshNodes/Private/Dataflow/ChaosFleshNodesPlugin.cpp
--- a/Engine/Plugins/Experimental/ChaosFlesh/Source/ChaosFleshNodes/Private/Dataflow/ChaosFleshNodesPlugin.cpp
+++ b/Engine/Plugins/Experimental/ChaosFlesh/Source/ChaosFleshNodes/Private/Dataflow/ChaosFleshNodesPlugin.cpp
@@ -19,7 +19,7 @@
 #include "Dataflow/ChaosFleshGetSurfaceIndicesNode.h"
 #include "Dataflow/ChaosFleshIsolateComponentNode.h"
 #include "Dataflow/ChaosFleshKinematicConstraintNode.h"
-#include "Dataflow/ChaosFleshKinematicOriginInsertionInitializationNode.h"
+#include "Dataflow/ChaosFleshKinematicMuscleAttachmentsNode.h"
 #include "Dataflow/ChaosFleshKinematicTetrahedralConstraintNode.h"
 #include "Dataflow/ChaosFleshRadialTetrahedronNode.h"
 #include "Dataflow/ChaosFleshTriangleMeshSimulationPropertiesNode.h"
@@ -58,7 +58,7 @@ void IChaosFleshNodesPlugin::StartupModule()
 	DATAFLOW_NODE_REGISTER_CREATION_FACTORY(FIsolateComponentNode);
 	DATAFLOW_NODE_REGISTER_CREATION_FACTORY(FKinematicBodySetupInitializationDataflowNode);
 	DATAFLOW_NODE_REGISTER_CREATION_FACTORY(FKinematicInitializationDataflowNode);
-	DATAFLOW_NODE_REGISTER_CREATION_FACTORY(FKinematicOriginInsertionInitializationDataflowNode);
+	DATAFLOW_NODE_REGISTER_CREATION_FACTORY(FKinematicMuscleAttachmentsDataflowNode);
 	DATAFLOW_NODE_REGISTER_CREATION_FACTORY(FKinematicSkeletalMeshInitializationDataflowNode);
 	DATAFLOW_NODE_REGISTER_CREATION_FACTORY(FKinematicTetrahedralBindingsDataflowNode);
 	DATAFLOW_NODE_REGISTER_CREATION_FACTORY(FRadialTetrahedronDataflowNodes);
diff --git a/Engine/Plugins/Experimental/ChaosFlesh/Source/ChaosFleshNodes/Public/Dataflow/ChaosFleshKinematicMuscleAttachmentsNode.h b/Engine/Plugins/Experimental/ChaosFlesh/Source/ChaosFleshNodes/Public/Dataflow/ChaosFleshKinematicMuscleAttachmentsNode.h
new file mode 100644
--- /dev/null
+++ b/Engine/Plugins/Experimental/ChaosFlesh/Source/ChaosFleshNodes/Public/Dataflow/ChaosFleshKinematicMuscleAttachmentsNode.h
@@ -0,0 +1,47 @@
+// Copyright Epic Games, Inc. All Rights Reserved.
+#pragma once 
+
+#include "CoreMinimal.h"
+#include "Dataflow/DataflowCore.h"
+#include "Dataflow/DataflowEngine.h"
+
+#include "ChaosFleshKinematicMuscleAttachmentsNode.generated.h"
+
+class USkeletalMesh;
+
+USTRUCT(meta = (DataflowFlesh))
+struct FKinematicMuscleAttachmentsDataflowNode : public FDataflowNode
+{
+	GENERATED_USTRUCT_BODY()
+	DATAFLOW_NODE_DEFINE_INTERNAL(FKinematicMuscleAttachmentsDataflowNode, "KinematicMuscleAttachments", "Flesh", "")
+	DATAFLOW_NODE_RENDER_TYPE(FGeometryCollection::StaticType(), "Collection")
+
+public:
+	typedef FManagedArrayCollection DataType;
+
+	UPROPERTY(meta = (DataflowInput, DataflowOutput, DisplayName = "Collection", DataflowPassthrough = "Collection"))
+		FManagedArrayCollection Collection;
+
+	UPROPERTY(meta = (DataflowInput, DisplayName = "OriginSelectionSet"))
+		TArray<int32> OriginVertexIndicesIn;
+
+	UPROPERTY(meta = (DataflowInput, DisplayName = "InsertionSelectionSet"))
+		TArray<int32> InsertionVertexIndicesIn;
+
+	UPROPERTY(meta = (DataflowInput, DisplayName = "SkeletalMesh"))
+		TObjectPtr<USkeletalMesh> SkeletalMeshIn;
+
+
+	FKinematicMuscleAttachmentsDataflowNode(const Dataflow::FNodeParameters& InParam, FGuid InGuid = FGuid::NewGuid())
+		: FDataflowNode(InParam, InGuid)
+	{
+		RegisterInputConnection(&Collection);
+		RegisterOutputConnection(&Collection, &Collection);
+		RegisterInputConnection(&OriginVertexIndicesIn);
+		RegisterInputConnection(&InsertionVertexIndicesIn);
+		RegisterInputConnection(&SkeletalMeshIn);
+	}
+
+	virtual void Evaluate(Dataflow::FContext& Context, const FDataflowOutput* Out) const override;
+};
+
diff --git a/Engine/Plugins/Experimental/ChaosFlesh/Source/ChaosFleshNodes/Public/Dataflow/ChaosFleshKinematicOriginInsertionInitializationNode.h b/Engine/Plugins/Experimental/ChaosFlesh/Source/ChaosFleshNodes/Public/Dataflow/ChaosFleshKinematicOriginInsertionInitializationNode.h
deleted file mode 100644
diff --git a/Engine/Plugins/Experimental/ChaosFlesh/Source/ChaosFleshNodes/Public/Dataflow/ChaosFleshSetFleshBonePositionTargetBindingNode.h b/Engine/Plugins/Experimental/ChaosFlesh/Source/ChaosFleshNodes/Public/Dataflow/ChaosFleshSetFleshBonePositionTargetBindingNode.h
--- a/Engine/Plugins/Experimental/ChaosFlesh/Source/ChaosFleshNodes/Public/Dataflow/ChaosFleshSetFleshBonePositionTargetBindingNode.h
+++ b/Engine/Plugins/Experimental/ChaosFlesh/Source/ChaosFleshNodes/Public/Dataflow/ChaosFleshSetFleshBonePositionTargetBindingNode.h
@@ -23,7 +23,7 @@ USTRUCT(meta = (DataflowFlesh))
 struct FSetFleshBonePositionTargetBindingDataflowNode : public FDataflowNode
 {
 	GENERATED_USTRUCT_BODY()
-		DATAFLOW_NODE_DEFINE_INTERNAL(FSetFleshBonePositionTargetBindingDataflowNode, "FSetFleshBonePositionTargetBinding", "Flesh", "")
+		DATAFLOW_NODE_DEFINE_INTERNAL(FSetFleshBonePositionTargetBindingDataflowNode, "SetFleshBonePositionTargetBinding", "Flesh", "")
 		DATAFLOW_NODE_RENDER_TYPE(FGeometryCollection::StaticType(), "Collection")
 
 public:
diff --git a/Engine/Plugins/Experimental/ChaosFlesh/Source/ChaosFleshNodes/Public/Dataflow/ChaosFleshSetVertexTetrahedraPositionTargetBindingNode.h b/Engine/Plugins/Experimental/ChaosFlesh/Source/ChaosFleshNodes/Public/Dataflow/ChaosFleshSetVertexTetrahedraPositionTargetBindingNode.h
--- a/Engine/Plugins/Experimental/ChaosFlesh/Source/ChaosFleshNodes/Public/Dataflow/ChaosFleshSetVertexTetrahedraPositionTargetBindingNode.h
+++ b/Engine/Plugins/Experimental/ChaosFlesh/Source/ChaosFleshNodes/Public/Dataflow/ChaosFleshSetVertexTetrahedraPositionTargetBindingNode.h
@@ -12,7 +12,7 @@ USTRUCT(meta = (DataflowFlesh))
 struct FSetVertexTetrahedraPositionTargetBindingDataflowNode : public FDataflowNode
 {
 	GENERATED_USTRUCT_BODY()
-		DATAFLOW_NODE_DEFINE_INTERNAL(FSetVertexTetrahedraPositionTargetBindingDataflowNode, "FSetVertexTetrahedraPositionTargetBinding", "Flesh", "")
+		DATAFLOW_NODE_DEFINE_INTERNAL(FSetVertexTetrahedraPositionTargetBindingDataflowNode, "SetVertexTetrahedraPositionTargetBinding", "Flesh", "")
 		DATAFLOW_NODE_RENDER_TYPE(FGeometryCollection::StaticType(), "Collection")
 
 public:
diff --git a/Engine/Plugins/Experimental/ChaosFlesh/Source/ChaosFleshNodes/Public/Dataflow/ChaosFleshSetVertexTrianglePositionTargetBindingNode.h b/Engine/Plugins/Experimental/ChaosFlesh/Source/ChaosFleshNodes/Public/Dataflow/ChaosFleshSetVertexTrianglePositionTargetBindingNode.h
--- a/Engine/Plugins/Experimental/ChaosFlesh/Source/ChaosFleshNodes/Public/Dataflow/ChaosFleshSetVertexTrianglePositionTargetBindingNode.h
+++ b/Engine/Plugins/Experimental/ChaosFlesh/Source/ChaosFleshNodes/Public/Dataflow/ChaosFleshSetVertexTrianglePositionTargetBindingNode.h
@@ -12,7 +12,7 @@ USTRUCT(meta = (DataflowFlesh))
 struct FSetVertexTrianglePositionTargetBindingDataflowNode : public FDataflowNode
 {
 	GENERATED_USTRUCT_BODY()
-		DATAFLOW_NODE_DEFINE_INTERNAL(FSetVertexTrianglePositionTargetBindingDataflowNode, "FSetVertexTrianglePositionTargetBinding", "Flesh", "")
+		DATAFLOW_NODE_DEFINE_INTERNAL(FSetVertexTrianglePositionTargetBindingDataflowNode, "SetVertexTrianglePositionTargetBinding", "Flesh", "")
 		DATAFLOW_NODE_RENDER_TYPE(FGeometryCollection::StaticType(), "Collection")
 
 public:
