ChaosFlesh : CreateTetrahedron updates

diff --git a/Engine/Plugins/Experimental/ChaosFlesh/Source/ChaosFleshNodes/Private/Dataflow/ChaosFleshCreateTetrahedronNode.cpp b/Engine/Plugins/Experimental/ChaosFlesh/Source/ChaosFleshNodes/Private/Dataflow/ChaosFleshCreateTetrahedronNode.cpp
--- a/Engine/Plugins/Experimental/ChaosFlesh/Source/ChaosFleshNodes/Private/Dataflow/ChaosFleshCreateTetrahedronNode.cpp
+++ b/Engine/Plugins/Experimental/ChaosFlesh/Source/ChaosFleshNodes/Private/Dataflow/ChaosFleshCreateTetrahedronNode.cpp
@@ -40,41 +40,70 @@ void FCreateTetrahedronDataflowNode::Evaluate(Dataflow::FContext& Context, const
 {
 	if (Out->IsA<DataType>(&Collection))
 	{
-		FFleshCollection GeneratedTetrahedron;
 		TUniquePtr<FFleshCollection> InCollection(GetValue<DataType>(Context, &Collection).NewCopy<FFleshCollection>());
+		TUniquePtr<FFleshCollection> InSourceCollection(GetValue<DataType>(Context, &SourceCollection).NewCopy<FFleshCollection>());
 
 		// TransformGroup
-		int32 NumTransforms = InCollection->NumElements(FTransformCollection::TransformGroup);
-		TManagedArray<FString>* TransformName = InCollection->FindAttribute<FString>("BoneName", FTransformCollection::TransformGroup);
-		TManagedArray<int32>* TransformToGeometryIndex = InCollection->FindAttribute<int32>("TransformToGeometryIndex", FTransformCollection::TransformGroup);
-		const TManagedArray<int32>* Parent = InCollection->FindAttribute<int32>(FTransformCollection::ParentAttribute, FTransformCollection::TransformGroup);
-		const TManagedArray<FTransform3f>* LocalSpaceTransform = InCollection->FindAttribute<FTransform3f>(FTransformCollection::TransformAttribute, FTransformCollection::TransformGroup);
+		int32 NumTransforms = InSourceCollection->NumElements(FTransformCollection::TransformGroup);
+		TManagedArray<FString>* TransformName = InSourceCollection->FindAttribute<FString>("BoneName", FTransformCollection::TransformGroup);
+		TManagedArray<int32>* TransformToGeometryIndex = InSourceCollection->FindAttribute<int32>("TransformToGeometryIndex", FTransformCollection::TransformGroup);
+		const TManagedArray<int32>* Parent = InSourceCollection->FindAttribute<int32>(FTransformCollection::ParentAttribute, FTransformCollection::TransformGroup);
+		const TManagedArray<FTransform3f>* LocalSpaceTransform = InSourceCollection->FindAttribute<FTransform3f>(FTransformCollection::TransformAttribute, FTransformCollection::TransformGroup);
 		// Geometry Group
-		int32 NumGeometry = InCollection->NumElements(FGeometryCollection::GeometryGroup);
-		TManagedArray<int32>* GroupToTransformIndex = InCollection->FindAttribute<int32>("TransformIndex", FGeometryCollection::GeometryGroup);
-		TManagedArray<int32>* VertexCount = InCollection->FindAttribute<int32>("VertexCount", FGeometryCollection::GeometryGroup);
-		TManagedArray<int32>* VertexStart = InCollection->FindAttribute<int32>("VertexStart", FGeometryCollection::GeometryGroup);
-		TManagedArray<int32>* FaceCount = InCollection->FindAttribute<int32>("FaceCount", FGeometryCollection::GeometryGroup);
-		TManagedArray<int32>* FaceStart = InCollection->FindAttribute<int32>("FaceStart", FGeometryCollection::GeometryGroup);
+		int32 NumGeometry = InSourceCollection->NumElements(FGeometryCollection::GeometryGroup);
+		TManagedArray<int32>* GroupToTransformIndex = InSourceCollection->FindAttribute<int32>("TransformIndex", FGeometryCollection::GeometryGroup);
+		TManagedArray<int32>* VertexCount = InSourceCollection->FindAttribute<int32>("VertexCount", FGeometryCollection::GeometryGroup);
+		TManagedArray<int32>* VertexStart = InSourceCollection->FindAttribute<int32>("VertexStart", FGeometryCollection::GeometryGroup);
+		TManagedArray<int32>* FaceCount = InSourceCollection->FindAttribute<int32>("FaceCount", FGeometryCollection::GeometryGroup);
+		TManagedArray<int32>* FaceStart = InSourceCollection->FindAttribute<int32>("FaceStart", FGeometryCollection::GeometryGroup);
 
 		// Vertices Group
-		int32 NumVertices = InCollection->NumElements(FGeometryCollection::VerticesGroup);
-		const TManagedArray<FVector3f>* Vertex = InCollection->FindAttribute<FVector3f>("Vertex", FGeometryCollection::VerticesGroup);
+		int32 NumVertices = InSourceCollection->NumElements(FGeometryCollection::VerticesGroup);
+		const TManagedArray<FVector3f>* Vertex = InSourceCollection->FindAttribute<FVector3f>("Vertex", FGeometryCollection::VerticesGroup);
 		// Faces Group
-		int32 NumTriangles = InCollection->NumElements(FGeometryCollection::FacesGroup);
-		const TManagedArray<FIntVector>* Faces = InCollection->FindAttribute<FIntVector>("Indices", FGeometryCollection::FacesGroup);
+		int32 NumTriangles = InSourceCollection->NumElements(FGeometryCollection::FacesGroup);
+		const TManagedArray<FIntVector>* Faces = InSourceCollection->FindAttribute<FIntVector>("Indices", FGeometryCollection::FacesGroup);
 
 		if (TransformName && TransformToGeometryIndex && Parent && LocalSpaceTransform &&
 			GroupToTransformIndex && VertexCount && VertexStart && FaceCount && FaceStart
 			&& Vertex && Faces)
 		{
-			TArray<int32> ProcessGeometryIndices = Dataflow::GetMatchingMeshIndices(MeshNames, InCollection.Get());
+			auto GetParentIndex = [&InCollection, &InSourceCollection](int32 Gdx)
+			{
+				if (Gdx >= 0)
+				{
+					TManagedArray<FString>* SourceTransformName = InSourceCollection->FindAttribute<FString>("BoneName", FTransformCollection::TransformGroup);
+					TManagedArray<FString>* TransformName = InCollection->FindAttribute<FString>("BoneName", FTransformCollection::TransformGroup);
+					TManagedArray<int32>* SourceGroupToTransformIndex = InSourceCollection->FindAttribute<int32>("TransformIndex", FGeometryCollection::GeometryGroup);
+					TManagedArray<int32>* GroupToTransformIndex = InCollection->FindAttribute<int32>("TransformIndex", FGeometryCollection::GeometryGroup);
+					TManagedArray<int32>* ParentIndex = InCollection->FindAttribute<int32>("Parent", FGeometryCollection::TransformGroup);
+					if (SourceTransformName && TransformName && SourceGroupToTransformIndex && GroupToTransformIndex && ParentIndex)
+					{
+						if (Gdx < SourceGroupToTransformIndex->Num() && Gdx < GroupToTransformIndex->Num())
+						{
+							int32 SrcTdx = (*SourceGroupToTransformIndex)[Gdx], Tdx = (*GroupToTransformIndex)[Gdx];
+							if (0 <= SrcTdx && SrcTdx < SourceTransformName->Num() && 0 <= Tdx && Tdx < TransformName->Num())
+							{
+								if ((*SourceTransformName)[SrcTdx].Equals((*TransformName)[Tdx]))
+								{
+									return (*ParentIndex)[Tdx];
+								}
+							}
+						}
+					}
+				}
+				return (int32)INDEX_NONE;
+			};
+
+			TArray<int32> ProcessGeometryIndices = Dataflow::GetMatchingMeshIndices(Selection, InSourceCollection.Get());
 
 			TArray<TUniquePtr<FFleshCollection>> CollectionBuffer;
 			for (int32 Gdx = 0; Gdx < NumGeometry; Gdx++)
 			{
-				if (ProcessGeometryIndices.Contains(Gdx)) 
-					CollectionBuffer.Add(TUniquePtr<FFleshCollection>( new FFleshCollection()));
+				if (ProcessGeometryIndices.Contains(Gdx))
+				{
+					CollectionBuffer.Add(TUniquePtr<FFleshCollection>(new FFleshCollection()));
+				}
 				else 
 					CollectionBuffer.Add(nullptr);
 			}
@@ -113,14 +142,29 @@ void FCreateTetrahedronDataflowNode::Evaluate(Dataflow::FContext& Context, const
 				}
 			});
 
-			auto AppendProcessedGeometry = [&ProcessGeometryIndices, &CollectionBuffer, &InCollection](FFleshCollection& ToCollection)
+			auto AppendProcessedGeometry = [&ProcessGeometryIndices, &CollectionBuffer, &InSourceCollection](FFleshCollection& ToCollection)
 			{
-				TManagedArray<int32>* SourceGroupToTransformIndex = InCollection->FindAttribute<int32>("TransformIndex", FGeometryCollection::GeometryGroup);
-				TManagedArray<FString>* SourceTransformName = InCollection->FindAttribute<FString>("BoneName", FTransformCollection::TransformGroup);
+				// find the index in the to transform group based on name
+				auto FindParentIndexFromName = [](FFleshCollection& fpColllection, FString fpName) 
+				{
+					if (!fpName.IsEmpty())
+					{
+						const TManagedArray<FString>& Names = fpColllection.GetAttribute<FString>("BoneName", FTransformCollection::TransformGroup);
+						for (int32 i = 0; i < Names.Num(); i++) if (Names[i].Equals(fpName)) return i;
+					}
+					return (int32)INDEX_NONE;
+				};
+
+				TManagedArray<int32>* SourceGroupToTransformIndex = InSourceCollection->FindAttribute<int32>("TransformIndex", FGeometryCollection::GeometryGroup);
+				TManagedArray<FString>* SourceTransformName = InSourceCollection->FindAttribute<FString>("BoneName", FTransformCollection::TransformGroup);
+				TManagedArray<int32>* SourceParentIndex = InSourceCollection->FindAttribute<int32>("Parent", FTransformCollection::TransformGroup);
 				TManagedArray<int32>* ToGroupToTransformIndex = ToCollection.FindAttribute<int32>("TransformIndex", FGeometryCollection::GeometryGroup);
+				TManagedArray<int32>* TransformToGroupIndex = ToCollection.FindAttribute<int32>("TransformToGeometryIndex", FTransformCollection::TransformGroup);
 				TManagedArray<FString>* ToTransformName = ToCollection.FindAttribute<FString>("BoneName", FTransformCollection::TransformGroup);
+				TManagedArray<int32>* ToParentIndex = ToCollection.FindAttribute<int32>("Parent", FTransformCollection::TransformGroup);
+				TManagedArray<TSet<int32> >* ToChildIndex = ToCollection.FindAttribute< TSet<int32> >(FTransformCollection::ChildrenAttribute, FTransformCollection::TransformGroup);
 
-				if (SourceGroupToTransformIndex && SourceTransformName && ToGroupToTransformIndex && ToTransformName)
+				if (SourceGroupToTransformIndex && SourceTransformName && ToGroupToTransformIndex && ToTransformName && ToParentIndex && SourceParentIndex && ToChildIndex)
 				{
 					for (int32 Sdx = 0; Sdx < ProcessGeometryIndices.Num(); Sdx++)
 					{
@@ -130,8 +174,10 @@ void FCreateTetrahedronDataflowNode::Evaluate(Dataflow::FContext& Context, const
 							int32 GeomIndex = ToCollection.NumElements(FGeometryCollection::GeometryGroup);
 							ToCollection.AppendGeometry(*CollectionBuffer[Gdx]);
 							// source data
-							int32 SourceNumTransforms = InCollection->NumElements(FGeometryCollection::TransformGroup);
+							int32 SourceNumTransforms = InSourceCollection->NumElements(FGeometryCollection::TransformGroup);
 							int32 SourceTransformIndex = (*SourceGroupToTransformIndex)[Gdx];
+							int32 SourceTransformParent = (0 <= SourceTransformIndex && SourceTransformIndex < SourceNumTransforms) ? (*SourceParentIndex)[SourceTransformIndex] : INDEX_NONE;
+							FString SourceParentName = (0 <= SourceTransformParent && SourceTransformParent < SourceNumTransforms) ? (*SourceTransformName)[SourceTransformParent] : FString("");
 							// target data
 							int32 ToNumTransforms = ToCollection.NumElements(FGeometryCollection::TransformGroup);
 							int32 ToGeomTransformIndex = (*ToGroupToTransformIndex)[GeomIndex];
@@ -146,15 +192,25 @@ void FCreateTetrahedronDataflowNode::Evaluate(Dataflow::FContext& Context, const
 							}
 							if (0 <= ToGeomTransformIndex && ToGeomTransformIndex < ToNumTransforms)
 							{
+								// set the name
 								(*ToTransformName)[ToGeomTransformIndex] = TetName;
+								// set transform to geometry and geometry to transform mappings
+								(*TransformToGroupIndex)[ToGeomTransformIndex] = GeomIndex;
+								(*ToGroupToTransformIndex)[GeomIndex] = ToGeomTransformIndex;
+								// set the parent and child mappings
+								(*ToParentIndex)[ToGeomTransformIndex] = FindParentIndexFromName(ToCollection, SourceParentName);
+								if ((*ToParentIndex)[ToGeomTransformIndex] != INDEX_NONE)
+								{
+									(*ToChildIndex)[(*ToParentIndex)[ToGeomTransformIndex]].Add(ToGeomTransformIndex);
+								}
 							}
 						}
 					}
 				}
 			};
-			AppendProcessedGeometry(GeneratedTetrahedron);
+			AppendProcessedGeometry(*InCollection);
 		}
-		SetValue<const DataType&>(Context, MoveTemp(GeneratedTetrahedron), &Collection);
+		SetValue<const DataType&>(Context, *InCollection, &Collection);
 	}
 }
 
diff --git a/Engine/Plugins/Experimental/ChaosFlesh/Source/ChaosFleshNodes/Public/Dataflow/ChaosFleshCreateTetrahedronNode.h b/Engine/Plugins/Experimental/ChaosFlesh/Source/ChaosFleshNodes/Public/Dataflow/ChaosFleshCreateTetrahedronNode.h
--- a/Engine/Plugins/Experimental/ChaosFlesh/Source/ChaosFleshNodes/Public/Dataflow/ChaosFleshCreateTetrahedronNode.h
+++ b/Engine/Plugins/Experimental/ChaosFlesh/Source/ChaosFleshNodes/Public/Dataflow/ChaosFleshCreateTetrahedronNode.h
@@ -27,19 +27,24 @@ struct FCreateTetrahedronDataflowNode : public FDataflowNode
 public:
 	typedef FManagedArrayCollection DataType;
 
+	// Tetrahedral meshing method (TetWild or ISO-stuffing)
 	UPROPERTY(EditAnywhere, Category = "Dataflow")
 	TEnumAsByte<TetMeshingMethod> Method = TetMeshingMethod::IsoStuffing;
 
+	// Name of the mesh in the collection. This is defined from the name of the mesh's
+	// parent transform. CollectionKey("BoneName","Vertices") 
 	UPROPERTY(EditAnywhere, Category = "Dataflow")
-	TArray<FString> MeshNames;
+	TArray<FString> Selection;
 
 	//
 	// IsoStuffing
 	//
 
+	// General control for density for the tetrahedron. 
 	UPROPERTY(EditAnywhere, Category = "Dataflow", meta = (ClampMin = "1", EditCondition = "Method == TetMeshingMethod::IsoStuffing", EditConditionHides))
 	int32 NumCells = 32;
 
+	// Surface offset percentage to increase or decrease the surface alignment. 
 	UPROPERTY(EditAnywhere, Category = "Dataflow", meta = (ClampMin = "-0.5", ClampMax = "0.5", EditCondition = "Method == TetMeshingMethod::IsoStuffing", EditConditionHides))
 	double OffsetPercent = 0.05;
 
@@ -85,14 +90,22 @@ public:
 	UPROPERTY(EditAnywhere, Category = "Dataflow")
 	bool bDiscardInteriorTriangles = true;
 
+	// Input pass-through collection. When connected, the generated tetrahedron will be nested into
+	// its associated parents transform.
 	UPROPERTY(meta = (DataflowInput, DataflowOutput, DisplayName = "Collection"))
 	FManagedArrayCollection Collection;
 
+	// Source collection used to generate the tetrahedron from. Closed geometry within the collections geometry group will be used to 
+	// generate tetrahedron.
+	UPROPERTY(meta = (DataflowInput, DisplayName = "SourceCollection"))
+	FManagedArrayCollection SourceCollection;
+
 	FCreateTetrahedronDataflowNode(const Dataflow::FNodeParameters& InParam, FGuid InGuid = FGuid::NewGuid())
 		: FDataflowNode(InParam, InGuid)
 	{
 		RegisterInputConnection(&Collection);
 		RegisterOutputConnection(&Collection);
+		RegisterInputConnection(&SourceCollection);
 	}
 
 	virtual void Evaluate(Dataflow::FContext& Context, const FDataflowOutput* Out) const override;
diff --git a/Engine/Source/Runtime/Experimental/GeometryCollectionEngine/Private/GeometryCollection/GeometryCollectionEngineConversion.cpp b/Engine/Source/Runtime/Experimental/GeometryCollectionEngine/Private/GeometryCollection/GeometryCollectionEngineConversion.cpp
--- a/Engine/Source/Runtime/Experimental/GeometryCollectionEngine/Private/GeometryCollection/GeometryCollectionEngineConversion.cpp
+++ b/Engine/Source/Runtime/Experimental/GeometryCollectionEngine/Private/GeometryCollection/GeometryCollectionEngineConversion.cpp
@@ -1273,6 +1273,11 @@ bool FGeometryCollectionEngineConversion::AppendSkeletalMesh(const USkeletalMesh
 	{
 		int32 ComponentTransformIndex = ComponentTransformBaseIndex + ComponentIndex;
 		Parent[ComponentTransformIndex] = RootIndex;
+		if (RootIndex != INDEX_NONE)
+		{
+			Children[RootIndex].Add(ComponentTransformIndex);
+		}
+
 		BoneName[ComponentTransformIndex] = FString::Printf(TEXT("%s_Mesh%d"), *BoneName[RootIndex], ComponentTransformIndex);
 		LocalSpaceTransform[ComponentTransformIndex] = FTransform3f::Identity;
 		SimulationType[ComponentTransformIndex] = FGeometryCollection::ESimulationTypes::FST_None;
