ChaosFlesh : ObjectState type change

diff --git a/Engine/Plugins/Experimental/ChaosFlesh/Source/ChaosFleshEngine/Private/ChaosFlesh/ChaosDeformableGameplayComponent.cpp b/Engine/Plugins/Experimental/ChaosFlesh/Source/ChaosFleshEngine/Private/ChaosFlesh/ChaosDeformableGameplayComponent.cpp
--- a/Engine/Plugins/Experimental/ChaosFlesh/Source/ChaosFleshEngine/Private/ChaosFlesh/ChaosDeformableGameplayComponent.cpp
+++ b/Engine/Plugins/Experimental/ChaosFlesh/Source/ChaosFleshEngine/Private/ChaosFlesh/ChaosDeformableGameplayComponent.cpp
@@ -103,7 +103,7 @@ void UDeformableGameplayComponent::DetectEnvironmentCollisions(const int32 MaxNu
 	const TManagedArray<TArray<int32>>* BoneIndices = BoneWeightsFacade.FindBoneIndices();
 	const TManagedArray<TArray<float>>* BoneWeights = BoneWeightsFacade.FindBoneWeights();
 
-	const TManagedArray<FTransform>* RestTransforms = nullptr;
+	const TManagedArray<FTransform3f>* RestTransforms = nullptr;
 	if (const UFleshAsset* FleshAsset = GetRestCollection())
 	{
 		if (const FFleshCollection* Rest = FleshAsset->GetCollection())
@@ -111,7 +111,7 @@ void UDeformableGameplayComponent::DetectEnvironmentCollisions(const int32 MaxNu
 			GeometryCollection::Facades::FTransformSource TransformSource(*Rest);
 			if (TransformSource.IsValid())
 			{
-				RestTransforms = Rest->FindAttribute<FTransform>(
+				RestTransforms = Rest->FindAttribute<FTransform3f>(
 					FTransformCollection::TransformAttribute, FTransformCollection::TransformGroup);
 			}
 		}
@@ -200,13 +200,13 @@ void UDeformableGameplayComponent::DetectEnvironmentCollisions(const int32 MaxNu
 					//FVector LocalPoint = RestInComponent.TransformPosition(ToDouble(RestPos));
 					
 					const FFleshCollection* Collection = GetRestCollection()->GetCollection();
-					FVector LocalPoint = (*RestTransforms)[Bones[j]].InverseTransformPosition(ToDouble(RestPos));
-					if (const TManagedArray<FTransform>* ConstTransformsPtr = Collection->FindAttributeTyped<FTransform>(FName("ComponentTransform"), FName("ComponentTransformGroup")))
+					FVector3f LocalPoint = (*RestTransforms)[Bones[j]].InverseTransformPosition(RestPos);
+					if (const TManagedArray<FTransform3f>* ConstTransformsPtr = Collection->FindAttributeTyped<FTransform3f>(FName("ComponentTransform"), FName("ComponentTransformGroup")))
 					{
-						LocalPoint = (*ConstTransformsPtr)[0].InverseTransformPosition(ToDouble(RestPos));
+						LocalPoint = (*ConstTransformsPtr)[0].InverseTransformPosition(RestPos);
 					}
 
-					SkinnedPos += AnimationInComponentSpaceTransforms[Bones[j]].TransformPosition(LocalPoint) * Weights[j];
+					SkinnedPos += AnimationInComponentSpaceTransforms[Bones[j]].TransformPosition(FVector(LocalPoint)) * Weights[j];
 					bDidSkinnedPosition = true;
 				}
 			}
diff --git a/Engine/Plugins/Experimental/ChaosFlesh/Source/ChaosFleshEngine/Private/ChaosFlesh/ChaosDeformableTetrahedralComponent.cpp b/Engine/Plugins/Experimental/ChaosFlesh/Source/ChaosFleshEngine/Private/ChaosFlesh/ChaosDeformableTetrahedralComponent.cpp
--- a/Engine/Plugins/Experimental/ChaosFlesh/Source/ChaosFleshEngine/Private/ChaosFlesh/ChaosDeformableTetrahedralComponent.cpp
+++ b/Engine/Plugins/Experimental/ChaosFlesh/Source/ChaosFleshEngine/Private/ChaosFlesh/ChaosDeformableTetrahedralComponent.cpp
@@ -151,10 +151,17 @@ UDeformablePhysicsComponent::FDataMapValue UDeformableTetrahedralComponent::NewD
 					TArray<USkeletalMeshComponent*> SkeletalMeshComponents;
 					GetOwner()->GetComponents<USkeletalMeshComponent>(SkeletalMeshComponents);
 
-					if (const TManagedArray<FTransform>* RestTransforms = Rest->FindAttribute<FTransform>(FTransformCollection::TransformAttribute, FTransformCollection::TransformGroup))
+					if (const TManagedArray<FTransform3f>* RestTransforms = Rest->FindAttribute<FTransform3f>(FTransformCollection::TransformAttribute, FTransformCollection::TransformGroup))
 					{
-						TArray<FTransform> AnimationTransforms = RestTransforms->GetConstArray();
-						TArray<FTransform> ComponentPose = RestTransforms->GetConstArray();
+						auto ToDoubleLocal = [](const TArray<FTransform3f>& Src, TArray<FTransform>& Tar) {
+							// @todo : Push floats through to the solver and avoid the copy. 
+							Tar.AddUninitialized(Src.Num());
+							for (int i = 0; i < Src.Num(); i++) Tar[i] = FTransform(Src[i]);
+						};
+
+						TArray<FTransform3f> FloatRestTransform = RestTransforms->GetConstArray();
+						TArray<FTransform> AnimationTransforms; ToDoubleLocal(FloatRestTransform, AnimationTransforms);
+						TArray<FTransform> ComponentPose; ToDoubleLocal(FloatRestTransform, ComponentPose);
 
 						// Extract animated transforms from all skeletal meshes.
 						for (const USkeletalMeshComponent* SkeletalMeshComponent : SkeletalMeshComponents)
@@ -196,7 +203,7 @@ UDeformablePhysicsComponent::FDataMapValue UDeformableTetrahedralComponent::NewD
 						FTransform BoneSpaceXf;
 						if (AnimationTransforms.IsValidIndex(SimulationSpace.SimSpaceTransformGlobalIndex))
 						{
-							BoneSpaceXf = AnimationTransforms[SimulationSpace.SimSpaceTransformGlobalIndex];
+							BoneSpaceXf = FTransform(AnimationTransforms[SimulationSpace.SimSpaceTransformGlobalIndex]);
 						}
 						else
 						{
diff --git a/Engine/Plugins/Experimental/ChaosFlesh/Source/ChaosFleshEngine/Private/ChaosFlesh/FleshDynamicAsset.cpp b/Engine/Plugins/Experimental/ChaosFlesh/Source/ChaosFleshEngine/Private/ChaosFlesh/FleshDynamicAsset.cpp
--- a/Engine/Plugins/Experimental/ChaosFlesh/Source/ChaosFleshEngine/Private/ChaosFlesh/FleshDynamicAsset.cpp
+++ b/Engine/Plugins/Experimental/ChaosFlesh/Source/ChaosFleshEngine/Private/ChaosFlesh/FleshDynamicAsset.cpp
@@ -21,7 +21,7 @@ UFleshDynamicAsset::UFleshDynamicAsset(const FObjectInitializer& ObjectInitializ
 void UFleshDynamicAsset::Init()
 {
 	DynamicCollection->AddAttribute<FVector3f>("Vertex", FGeometryCollection::VerticesGroup);
-	DynamicCollection->AddAttribute<FVector3f>("ObjectState", FGeometryCollection::VerticesGroup);
+	DynamicCollection->AddAttribute<int32>("ObjectState", FGeometryCollection::VerticesGroup);
 }
 
 void UFleshDynamicAsset::Reset(const FManagedArrayCollection* InCopyFrom)
diff --git a/Engine/Plugins/Experimental/ChaosFlesh/Source/ChaosFleshEngine/Private/ChaosFlesh/SimulationAsset.cpp b/Engine/Plugins/Experimental/ChaosFlesh/Source/ChaosFleshEngine/Private/ChaosFlesh/SimulationAsset.cpp
--- a/Engine/Plugins/Experimental/ChaosFlesh/Source/ChaosFleshEngine/Private/ChaosFlesh/SimulationAsset.cpp
+++ b/Engine/Plugins/Experimental/ChaosFlesh/Source/ChaosFleshEngine/Private/ChaosFlesh/SimulationAsset.cpp
@@ -20,7 +20,7 @@ USimulationAsset::USimulationAsset(const FObjectInitializer& ObjectInitializer)
 
 void USimulationAsset::Init()
 {
-	SimulationCollection->AddAttribute<FVector3f>("ObjectState", FGeometryCollection::VerticesGroup);
+	SimulationCollection->AddAttribute<int32>("ObjectState", FGeometryCollection::VerticesGroup);
 }
 
 void USimulationAsset::Reset(const FManagedArrayCollection* InCopyFrom)
