ChaosFlesh  : MoveTemp hand off of arrays to solver.

diff --git a/Engine/Plugins/Experimental/ChaosFlesh/Source/ChaosFleshEngine/Private/ChaosFlesh/ChaosDeformableTetrahedralComponent.cpp b/Engine/Plugins/Experimental/ChaosFlesh/Source/ChaosFleshEngine/Private/ChaosFlesh/ChaosDeformableTetrahedralComponent.cpp
--- a/Engine/Plugins/Experimental/ChaosFlesh/Source/ChaosFleshEngine/Private/ChaosFlesh/ChaosDeformableTetrahedralComponent.cpp
+++ b/Engine/Plugins/Experimental/ChaosFlesh/Source/ChaosFleshEngine/Private/ChaosFlesh/ChaosDeformableTetrahedralComponent.cpp
@@ -212,12 +212,12 @@ UDeformablePhysicsComponent::FDataMapValue UDeformableTetrahedralComponent::NewD
 
 						return FDataMapValue(
 							new Chaos::Softs::FFleshThreadingProxy::FFleshInputBuffer(
-								*GetSimulationCollection()->GetCollection(),
+								MoveTemp(*GetSimulationCollection()->GetCollection()),
 								this->GetComponentTransform(),
 								BoneSpaceXf, 
 								SimulationSpace.SimSpaceTransformGlobalIndex,
-								AnimationTransforms, 
-								ComponentPose, 
+								MoveTemp(AnimationTransforms), 
+								MoveTemp(ComponentPose), 
 								BodyForces.bApplyGravity,
 								BodyForces.StiffnessMultiplier,
 								BodyForces.DampingMultiplier,
@@ -232,7 +232,7 @@ UDeformablePhysicsComponent::FDataMapValue UDeformableTetrahedralComponent::NewD
 	}
 	return FDataMapValue(
 		new Chaos::Softs::FFleshThreadingProxy::FFleshInputBuffer(
-			*GetSimulationCollection()->GetCollection(),
+			MoveTemp(*GetSimulationCollection()->GetCollection()),
 			this->GetComponentTransform(),
 			GetSimSpaceRestTransform(),
 			SimulationSpace.SimSpaceTransformGlobalIndex,
diff --git a/Engine/Source/Runtime/Experimental/Chaos/Public/Chaos/Deformable/ChaosDeformableSolverProxy.h b/Engine/Source/Runtime/Experimental/Chaos/Public/Chaos/Deformable/ChaosDeformableSolverProxy.h
--- a/Engine/Source/Runtime/Experimental/Chaos/Public/Chaos/Deformable/ChaosDeformableSolverProxy.h
+++ b/Engine/Source/Runtime/Experimental/Chaos/Public/Chaos/Deformable/ChaosDeformableSolverProxy.h
@@ -252,7 +252,7 @@ namespace Chaos::Softs
 			typedef FFleshThreadingProxy Source;
 
 			FFleshInputBuffer(
-				const FManagedArrayCollection InSimulationCollection,
+				const FManagedArrayCollection&& InSimulationCollection,
 				const FTransform& InWorldToComponentXf,
 				const FTransform& InComponentToBoneXf,
 				const int32 InSimSpaceBoneIndex,
@@ -279,12 +279,12 @@ namespace Chaos::Softs
 			{}
 
 			FFleshInputBuffer(
-				const FManagedArrayCollection InSimulationCollection,
+				const FManagedArrayCollection&& InSimulationCollection,
 				const FTransform& InWorldToComponentXf,
 				const FTransform& InComponentToBoneXf, 
 				const int32 InSimSpaceBoneIndex,
-				const TArray<FTransform>& InTransforms, 
-				const TArray<FTransform>& InRestTransforms, 
+				const TArray<FTransform>&& InTransforms, 
+				const TArray<FTransform>&& InRestTransforms, 
 				const bool InbEnableGravity, 
 				const float InStiffnessMultiplier, 
 				const float InDampingMultiplier, 
