ChaosFlesh : Add geometry with transform compensation.

diff --git a/Engine/Plugins/Experimental/ChaosFlesh/Source/ChaosFlesh/Private/ChaosFlesh/ChaosFleshCollectionFacade.cpp b/Engine/Plugins/Experimental/ChaosFlesh/Source/ChaosFlesh/Private/ChaosFlesh/ChaosFleshCollectionFacade.cpp
new file mode 100644
--- /dev/null
+++ b/Engine/Plugins/Experimental/ChaosFlesh/Source/ChaosFlesh/Private/ChaosFlesh/ChaosFleshCollectionFacade.cpp
@@ -0,0 +1,92 @@
+// Copyright Epic Games, Inc. All Rights Reserved.
+
+#include "ChaosFlesh/ChaosFleshCollectionFacade.h"
+
+#include "GeometryCollection/GeometryCollectionAlgo.h"
+#include "GeometryCollection/TransformCollection.h"
+#include "GeometryCollection/GeometryCollection.h"
+#include "ChaosFlesh/TetrahedralCollection.h"
+
+#define LOCTEXT_NAMESPACE "FFleshCollectionFacade"
+
+namespace Chaos 
+{
+
+	FFleshCollectionFacade::FFleshCollectionFacade(FManagedArrayCollection& InCollection)
+		: ConstCollection(InCollection)
+		, Collection(&InCollection)
+		, BoneName(InCollection, "BoneName", FTransformCollection::TransformGroup)
+		, Transform(InCollection, FTransformCollection::TransformAttribute, FTransformCollection::TransformGroup)
+		, Parent(InCollection, FTransformCollection::ParentAttribute, FTransformCollection::TransformGroup)
+		, Child(InCollection, FTransformCollection::ChildrenAttribute, FTransformCollection::TransformGroup)
+		, BoneMap(InCollection, FName("BoneMap"), FName("Vertices"))
+		, Vertex(InCollection, FName("Vertex"), FName("Vertices"))
+		, Indices(InCollection, FName("Indices"), FName("Faces"))
+		, Tetrahedron(InCollection, FTetrahedralCollection::TetrahedronAttribute, FTetrahedralCollection::TetrahedralGroup)
+		, VertexStart(InCollection, "VertexStart", FGeometryCollection::GeometryGroup)
+		, VertexCount(InCollection, "VertexCount", FGeometryCollection::GeometryGroup)
+		, FaceStart(InCollection, "FaceStart", FGeometryCollection::GeometryGroup)
+		, FaceCount(InCollection, "FaceCount", FGeometryCollection::GeometryGroup)
+	{}
+
+	FFleshCollectionFacade::FFleshCollectionFacade(const FManagedArrayCollection& InCollection)
+		: ConstCollection(InCollection)
+		, Collection(nullptr)
+		, BoneName(InCollection, "BoneName", FTransformCollection::TransformGroup)
+		, Transform(InCollection, FTransformCollection::TransformAttribute, FTransformCollection::TransformGroup)
+		, Parent(InCollection, FTransformCollection::ParentAttribute, FTransformCollection::TransformGroup)
+		, Child(InCollection, FTransformCollection::ChildrenAttribute, FTransformCollection::TransformGroup)
+		, BoneMap(InCollection, FName("BoneMap"), FName("Vertices"))
+		, Vertex(InCollection, FName("Vertex"), FName("Vertices"))
+		, Indices(InCollection, FName("Indices"), FName("Faces"))
+		, Tetrahedron(InCollection, FTetrahedralCollection::TetrahedronAttribute, FTetrahedralCollection::TetrahedralGroup)
+		, VertexStart(InCollection, "VertexStart", FGeometryCollection::GeometryGroup)
+		, VertexCount(InCollection, "VertexCount", FGeometryCollection::GeometryGroup)
+		, FaceStart(InCollection, "FaceStart", FGeometryCollection::GeometryGroup)
+		, FaceCount(InCollection, "FaceCount", FGeometryCollection::GeometryGroup)
+	{}
+
+	bool FFleshCollectionFacade::IsCompletelyValid() const {
+		return BoneName.IsValid() && Transform.IsValid() && Parent.IsValid() && Child.IsValid() &&
+			BoneMap.IsValid() && Vertex.IsValid() &&
+			Indices.IsValid() && Tetrahedron.IsValid() &&
+			VertexStart.IsValid() && VertexCount.IsValid() && FaceStart.IsValid() && FaceCount.IsValid();
+	}
+
+	bool FFleshCollectionFacade::IsTetrahedronValid() const {
+		return Vertex.IsValid() && Tetrahedron.IsValid();
+	}
+
+	bool FFleshCollectionFacade::IsHierarchyValid() const {
+		return Transform.IsValid() && Parent.IsValid();
+	}
+
+	void FFleshCollectionFacade::ComponentSpaceVertices(TArray<FVector3f>& OutComponentSpaceVertices)
+	{
+		ComponentSpaceVertices(OutComponentSpaceVertices, 0, Vertex.Num());
+	}
+
+	void FFleshCollectionFacade::ComponentSpaceVertices(TArray<FVector3f>& OutComponentSpaceVertices, int32 Start, int32 Count)
+	{
+		if (IsHierarchyValid() && Vertex.IsValid())
+		{
+			TArray<FTransform3f> ComponentTransform;
+			GeometryCollectionAlgo::GlobalMatrices(Transform.Get(), Parent.Get(), ComponentTransform);
+
+			OutComponentSpaceVertices.SetNumUninitialized(Count);
+			for (int i = 0; i < Count; i++)
+			{
+				int j = i + Start;
+				if (0 < BoneMap[i] && BoneMap[i] < ComponentTransform.Num())
+				{
+					OutComponentSpaceVertices[i] = ComponentTransform[BoneMap[j]].TransformPosition(Vertex[j]);
+				}
+				else
+				{
+					OutComponentSpaceVertices[i] = Vertex[j];
+				}
+			}
+		}
+	}
+
+}
diff --git a/Engine/Plugins/Experimental/ChaosFlesh/Source/ChaosFlesh/Public/ChaosFlesh/ChaosFleshCollectionFacade.h b/Engine/Plugins/Experimental/ChaosFlesh/Source/ChaosFlesh/Public/ChaosFlesh/ChaosFleshCollectionFacade.h
new file mode 100644
--- /dev/null
+++ b/Engine/Plugins/Experimental/ChaosFlesh/Source/ChaosFlesh/Public/ChaosFlesh/ChaosFleshCollectionFacade.h
@@ -0,0 +1,58 @@
+// Copyright Epic Games, Inc. All Rights Reserved.
+
+#include "CoreMinimal.h"
+#include "GeometryCollection/ManagedArrayCollection.h"
+#include "GeometryCollection/ManagedArrayAccessor.h"
+
+
+namespace Chaos {
+
+	class CHAOSFLESH_API FFleshCollectionFacade
+	{
+		const FManagedArrayCollection& ConstCollection;
+		FManagedArrayCollection* Collection = nullptr;
+	public:
+
+		FFleshCollectionFacade(FManagedArrayCollection& InCollection);
+
+		FFleshCollectionFacade(const FManagedArrayCollection& InCollection);
+		
+		/*Are all the public attributes avaiable.*/
+		bool IsCompletelyValid() const;
+	
+		/*Has tetrahedral attributes {Tetrahderon, Vertices}.*/
+		bool IsTetrahedronValid() const;
+
+		/*Has hierarchy attributes {Transform, Parent}.*/
+		bool IsHierarchyValid() const;
+
+		/*View a Attribute on the collection*/
+		template<class T>
+		const TManagedArray<T>* FindAttribute(FString AttributeName, FString Group) const;
+
+		/*Edit a Attribute on the collection*/
+		template<class T>
+		TManagedArray<T>* ModifyAttribute(FString AttributeName, FString Group);
+
+		/* All the vertices mapped into component space. */
+		void ComponentSpaceVertices(TArray<FVector3f>& OutComponentSpaceVertices);
+
+		/* Range of vertices mapped into component space*/
+		void ComponentSpaceVertices(TArray<FVector3f>& OutComponentSpaceVertices, int32 Start, int32 Count);
+
+		/*Public Attributes*/
+		TManagedArrayAccessor<FString> BoneName;
+		TManagedArrayAccessor<FTransform3f> Transform;
+		TManagedArrayAccessor<int32> Parent;
+		TManagedArrayAccessor< TSet<int32> > Child;
+		TManagedArrayAccessor<int32> BoneMap;
+		TManagedArrayAccessor<FVector3f> Vertex;
+		TManagedArrayAccessor<FIntVector3> Indices;
+		TManagedArrayAccessor<FIntVector4> Tetrahedron;
+		TManagedArrayAccessor<int32> VertexStart;
+		TManagedArrayAccessor<int32> VertexCount;
+		TManagedArrayAccessor<int32> FaceStart;
+		TManagedArrayAccessor<int32> FaceCount;
+	};
+
+}
diff --git a/Engine/Plugins/Experimental/ChaosFlesh/Source/ChaosFleshDeprecatedNodes/Private/Dataflow/ChaosFleshKinematicTetrahedralConstraintNode.cpp b/Engine/Plugins/Experimental/ChaosFlesh/Source/ChaosFleshDeprecatedNodes/Private/Dataflow/ChaosFleshKinematicTetrahedralConstraintNode.cpp
new file mode 100644
--- /dev/null
+++ b/Engine/Plugins/Experimental/ChaosFlesh/Source/ChaosFleshDeprecatedNodes/Private/Dataflow/ChaosFleshKinematicTetrahedralConstraintNode.cpp
@@ -0,0 +1,122 @@
+// Copyright Epic Games, Inc. All Rights Reserved.
+
+#include "Dataflow/ChaosFleshKinematicTetrahedralConstraintNode.h"
+
+#include "ChaosFlesh/TetrahedralCollection.h"
+#include "Engine/SkeletalMesh.h"
+#include "GeometryCollection/Facades/CollectionKinematicBindingFacade.h"
+#include "GeometryCollection/Facades/CollectionVertexBoneWeightsFacade.h"
+
+#include UE_INLINE_GENERATED_CPP_BY_NAME(ChaosFleshKinematicTetrahedralConstraintNode)
+
+
+void FKinematicTetrahedralBindingsDataflowNode::Evaluate(Dataflow::FContext& Context, const FDataflowOutput* Out) const
+{
+	if (Out->IsA<DataType>(&Collection))
+	{
+		DataType InCollection = GetValue<DataType>(Context, &Collection);
+		TManagedArray<FIntVector4>* Tetrahedron = InCollection.FindAttribute<FIntVector4>(FTetrahedralCollection::TetrahedronAttribute, FTetrahedralCollection::TetrahedralGroup);
+		TManagedArray<FVector3f>* Vertex = InCollection.FindAttribute<FVector3f>("Vertex", "Vertices");
+
+		TObjectPtr<USkeletalMesh> SkeletalMesh = GetValue<TObjectPtr<USkeletalMesh>>(Context, &SkeletalMeshIn);
+		if (SkeletalMesh && Tetrahedron && Vertex)
+		{
+			//parse exclusion list to find bones to skip
+			TArray<FString> StrArray;
+			ExclusionList.ParseIntoArray(StrArray, *FString(" "));				
+			int32 NumTets = Tetrahedron->Num();
+			TArray<FTransform> ComponentPose;
+			Dataflow::Animation::GlobalTransforms(SkeletalMesh->GetRefSkeleton(), ComponentPose);
+			TArray<bool> VertAdded;
+			VertAdded.Init(false, Vertex->Num());
+			auto DoSkipBoneIndex = [&StrArray, &SkeletalMesh](int32 BoneIndex)
+			{
+				bool Skip = false;
+				FString BoneName = SkeletalMesh->GetRefSkeleton().GetBoneName(BoneIndex).ToString();
+				for (FString Elem : StrArray)
+				{
+					if (BoneName.Contains(Elem))
+					{
+						Skip = true;
+						break;
+					}
+				}
+				return Skip;
+			};
+			for (int32 b = 0; b < SkeletalMesh->GetRefSkeleton().GetNum(); ++b)
+			{
+				FVector3f BonePosition(ComponentPose[b].GetTranslation());
+				int32 ParentIndex=SkeletalMesh->GetRefSkeleton().GetParentIndex(b);
+				if (!(ParentIndex == INDEX_NONE || DoSkipBoneIndex(b) || DoSkipBoneIndex(ParentIndex)))
+				{
+					FVector3f ParentPosition(ComponentPose[ParentIndex].GetTranslation());
+					FVector3f RayDir = ParentPosition - BonePosition;
+					Chaos::FReal Length = RayDir.Length();
+					RayDir.Normalize();
+
+					if (Length > Chaos::FReal(1e-8)) 
+					{
+						TSet<int32> BoneVertSet;
+						for (int32 t = 0; t < NumTets; t++) 
+						{
+							int32 i = (*Tetrahedron)[t][0];
+							int32 j = (*Tetrahedron)[t][1];
+							int32 k = (*Tetrahedron)[t][2];
+							int32 l = (*Tetrahedron)[t][3];
+
+							TArray<Chaos::TVec3<Chaos::FRealSingle>> InVertices;
+							InVertices.SetNum(4);
+							InVertices[0][0] = (*Vertex)[i].X; InVertices[0][1] = (*Vertex)[i].Y; InVertices[0][2] = (*Vertex)[i].Z;
+							InVertices[1][0] = (*Vertex)[j].X; InVertices[1][1] = (*Vertex)[j].Y; InVertices[1][2] = (*Vertex)[j].Z;
+							InVertices[2][0] = (*Vertex)[k].X; InVertices[2][1] = (*Vertex)[k].Y; InVertices[2][2] = (*Vertex)[k].Z;
+							InVertices[3][0] = (*Vertex)[l].X; InVertices[3][1] = (*Vertex)[l].Y; InVertices[3][2] = (*Vertex)[l].Z;
+							Chaos::FConvex ConvexTet(InVertices, Chaos::FReal(0));
+							Chaos::FReal OutTime;
+							Chaos::FVec3 OutPosition, OutNormal;
+							int32 OutFaceIndex;
+							bool KeepTet = ConvexTet.Raycast(BonePosition, RayDir, Length, Chaos::FReal(0), OutTime, OutPosition, OutNormal, OutFaceIndex);
+							if (KeepTet) 
+							{	
+								for (int32 c = 0; c < 4; ++c)
+								{	
+									if (!VertAdded[(*Tetrahedron)[t][c]])
+									{
+										VertAdded[(*Tetrahedron)[t][c]] = true;
+										BoneVertSet.Add((*Tetrahedron)[t][c]);
+									}
+								}
+							}
+						}
+
+						TArray<int32> BoundVerts = BoneVertSet.Array();
+						TArray<float> BoundWeights;
+						BoundWeights.Init(float(1), BoundVerts.Num());
+						if (BoundVerts.Num())
+						{
+							//get local coords of bound verts
+							typedef GeometryCollection::Facades::FKinematicBindingFacade FKinematics;
+							FKinematics Kinematics(InCollection); Kinematics.DefineSchema();
+							if (Kinematics.IsValid())
+							{
+								FKinematics::FBindingKey Binding = Kinematics.SetBoneBindings(ParentIndex, BoundVerts, BoundWeights);
+								TManagedArray<TArray<FVector3f>>& LocalPos = InCollection.AddAttribute<TArray<FVector3f>>("LocalPosition", Binding.GroupName);
+								Kinematics.AddKinematicBinding(Binding);
+								auto DoubleVert = [](FVector3f V) { return FVector3d(V.X, V.Y, V.Z); };
+								auto FloatVert = [](FVector3d V) { return FVector3f(V.X, V.Y, V.Z); };
+								LocalPos[Binding.Index].SetNum(BoundVerts.Num());
+								for (int32 i = 0; i < BoundVerts.Num(); i++)
+								{
+									FVector3f Temp = (*Vertex)[BoundVerts[i]];
+									LocalPos[Binding.Index][i] = FloatVert(ComponentPose[ParentIndex].InverseTransformPosition(DoubleVert(Temp)));
+								}
+							}
+						}
+					}
+				}
+			}
+			GeometryCollection::Facades::FVertexBoneWeightsFacade(InCollection).AddBoneWeightsFromKinematicBindings();
+		}
+
+		SetValue(Context, MoveTemp(InCollection), &Collection);
+	}
+}
diff --git a/Engine/Plugins/Experimental/ChaosFlesh/Source/ChaosFleshDeprecatedNodes/Public/Dataflow/ChaosFleshKinematicTetrahedralConstraintNode.h b/Engine/Plugins/Experimental/ChaosFlesh/Source/ChaosFleshDeprecatedNodes/Public/Dataflow/ChaosFleshKinematicTetrahedralConstraintNode.h
new file mode 100644
--- /dev/null
+++ b/Engine/Plugins/Experimental/ChaosFlesh/Source/ChaosFleshDeprecatedNodes/Public/Dataflow/ChaosFleshKinematicTetrahedralConstraintNode.h
@@ -0,0 +1,43 @@
+// Copyright Epic Games, Inc. All Rights Reserved.
+#pragma once 
+
+#include "CoreMinimal.h"
+#include "Dataflow/DataflowCore.h"
+#include "Dataflow/DataflowEngine.h"
+
+#include "ChaosFleshKinematicTetrahedralConstraintNode.generated.h"
+
+class USkeletalMesh;
+
+// @todo(deprecate), rename to FKinematicTetrahedralConstraintDataflowNode
+// ... This should really be a mode on the KinematicConstraint. 
+USTRUCT(meta = (DataflowFlesh))
+struct FKinematicTetrahedralBindingsDataflowNode : public FDataflowNode
+{
+	GENERATED_USTRUCT_BODY()
+	DATAFLOW_NODE_DEFINE_INTERNAL(FKinematicTetrahedralBindingsDataflowNode, "KinematicTetrahedralBindings", "Flesh", "")
+	DATAFLOW_NODE_RENDER_TYPE(FGeometryCollection::StaticType(), "Collection")
+
+public:
+	typedef FManagedArrayCollection DataType;
+
+	UPROPERTY(meta = (DataflowInput, DisplayName = "SkeletalMesh"))
+	TObjectPtr<USkeletalMesh> SkeletalMeshIn = nullptr;
+
+	UPROPERTY(meta = (DataflowInput, DataflowOutput, DisplayName = "Collection", DataflowPassthrough  = "Collection"))
+	FManagedArrayCollection Collection;
+
+	UPROPERTY(EditAnywhere, Category = "Dataflow")
+	FString ExclusionList = "twist foo";
+
+	FKinematicTetrahedralBindingsDataflowNode(const Dataflow::FNodeParameters& InParam, FGuid InGuid = FGuid::NewGuid())
+		: FDataflowNode(InParam, InGuid)
+	{
+		RegisterInputConnection(&SkeletalMeshIn);
+		RegisterInputConnection(&Collection);
+		RegisterOutputConnection(&Collection, &Collection);
+	}
+
+	virtual void Evaluate(Dataflow::FContext& Context, const FDataflowOutput* Out) const override;
+};
+
diff --git a/Engine/Plugins/Experimental/ChaosFlesh/Source/ChaosFleshNodes/Private/Dataflow/ChaosFleshCreateTetrahedronNode.cpp b/Engine/Plugins/Experimental/ChaosFlesh/Source/ChaosFleshNodes/Private/Dataflow/ChaosFleshCreateTetrahedronNode.cpp
--- a/Engine/Plugins/Experimental/ChaosFlesh/Source/ChaosFleshNodes/Private/Dataflow/ChaosFleshCreateTetrahedronNode.cpp
+++ b/Engine/Plugins/Experimental/ChaosFlesh/Source/ChaosFleshNodes/Private/Dataflow/ChaosFleshCreateTetrahedronNode.cpp
@@ -35,7 +35,6 @@
 // FCreateTetrahedronDataflowNode
 //=============================================================================
 
-
 void FCreateTetrahedronDataflowNode::Evaluate(Dataflow::FContext& Context, const FDataflowOutput* Out) const
 {
 	if (Out->IsA<DataType>(&Collection))
@@ -161,8 +160,12 @@ void FCreateTetrahedronDataflowNode::Evaluate(Dataflow::FContext& Context, const
 				TManagedArray<int32>* ToGroupToTransformIndex = ToCollection.FindAttribute<int32>("TransformIndex", FGeometryCollection::GeometryGroup);
 				TManagedArray<int32>* TransformToGroupIndex = ToCollection.FindAttribute<int32>("TransformToGeometryIndex", FTransformCollection::TransformGroup);
 				TManagedArray<FString>* ToTransformName = ToCollection.FindAttribute<FString>("BoneName", FTransformCollection::TransformGroup);
+				TManagedArray<FTransform3f>* ToTransform = ToCollection.FindAttribute<FTransform3f>("Transform", FTransformCollection::TransformGroup);
 				TManagedArray<int32>* ToParentIndex = ToCollection.FindAttribute<int32>("Parent", FTransformCollection::TransformGroup);
 				TManagedArray<TSet<int32> >* ToChildIndex = ToCollection.FindAttribute< TSet<int32> >(FTransformCollection::ChildrenAttribute, FTransformCollection::TransformGroup);
+				TManagedArray<FVector3f>* ToVertex = ToCollection.FindAttribute<FVector3f>("Vertex", FGeometryCollection::VerticesGroup);
+				TManagedArray<int32>* ToVertexCount = ToCollection.FindAttribute<int32>("VertexCount", FGeometryCollection::GeometryGroup);
+				TManagedArray<int32>* ToVertexStart = ToCollection.FindAttribute<int32>("VertexStart", FGeometryCollection::GeometryGroup);
 
 				if (SourceGroupToTransformIndex && SourceTransformName && ToGroupToTransformIndex && ToTransformName && ToParentIndex && SourceParentIndex && ToChildIndex)
 				{
@@ -194,6 +197,7 @@ void FCreateTetrahedronDataflowNode::Evaluate(Dataflow::FContext& Context, const
 							{
 								// set the name
 								(*ToTransformName)[ToGeomTransformIndex] = TetName;
+				
 								// set transform to geometry and geometry to transform mappings
 								(*TransformToGroupIndex)[ToGeomTransformIndex] = GeomIndex;
 								(*ToGroupToTransformIndex)[GeomIndex] = ToGeomTransformIndex;
@@ -203,6 +207,16 @@ void FCreateTetrahedronDataflowNode::Evaluate(Dataflow::FContext& Context, const
 								{
 									(*ToChildIndex)[(*ToParentIndex)[ToGeomTransformIndex]].Add(ToGeomTransformIndex);
 								}
+
+								if (ToGeomTransformIndex != INDEX_NONE)
+								{
+									int32 VertexEnd = (*ToVertexStart)[GeomIndex] + (*ToVertexCount)[GeomIndex];
+									FTransform3f ParentTransform = GeometryCollectionAlgo::GlobalMatrix3f(*ToTransform, *ToParentIndex, ToGeomTransformIndex);
+									for (int Vdx = (*ToVertexStart)[GeomIndex]; Vdx < VertexEnd; Vdx++)
+									{
+										(*ToVertex)[Vdx] = ParentTransform.InverseTransformPosition((*ToVertex)[Vdx]);
+									}
+								}
 							}
 						}
 					}
diff --git a/Engine/Plugins/Experimental/ChaosFlesh/Source/ChaosFleshNodes/Private/Dataflow/ChaosFleshKinematicSkeletonConstraintNode.cpp b/Engine/Plugins/Experimental/ChaosFlesh/Source/ChaosFleshNodes/Private/Dataflow/ChaosFleshKinematicSkeletonConstraintNode.cpp
new file mode 100644
--- /dev/null
+++ b/Engine/Plugins/Experimental/ChaosFlesh/Source/ChaosFleshNodes/Private/Dataflow/ChaosFleshKinematicSkeletonConstraintNode.cpp
@@ -0,0 +1,126 @@
+// Copyright Epic Games, Inc. All Rights Reserved.
+
+#include "Dataflow/ChaosFleshKinematicSkeletonConstraintNode.h"
+
+#include "ChaosFlesh/ChaosFleshCollectionFacade.h"
+#include "ChaosFlesh/TetrahedralCollection.h"
+#include "Animation/Skeleton.h"
+#include "GeometryCollection/Facades/CollectionKinematicBindingFacade.h"
+#include "GeometryCollection/Facades/CollectionVertexBoneWeightsFacade.h"
+#include "GeometryCollection/GeometryCollectionAlgo.h"
+
+#include UE_INLINE_GENERATED_CPP_BY_NAME(ChaosFleshKinematicSkeletonConstraintNode)
+
+void FKinematicSkeletonConstraintDataflowNode::Evaluate(Dataflow::FContext& Context, const FDataflowOutput* Out) const
+{
+	if (Out->IsA<DataType>(&Collection))
+	{
+		DataType InCollection = GetValue<DataType>(Context, &Collection);
+		Chaos::FFleshCollectionFacade TetCollection(InCollection);
+		TObjectPtr<USkeleton> Skeleton = GetValue<TObjectPtr<USkeleton>>(Context, &SkeletonIn);
+		if (Skeleton && TetCollection.IsTetrahedronValid() )
+		{
+			//parse exclusion list to find bones to skip
+			TArray<FString> StrArray;
+			ExclusionList.ParseIntoArray(StrArray, *FString(" "));				
+			int32 NumTets = TetCollection.Tetrahedron.Num();
+
+			TArray<FTransform> ComponentPose;
+			Dataflow::Animation::GlobalTransforms(Skeleton->GetReferenceSkeleton(), ComponentPose);
+
+			TArray<FVector3f> Vertex = TetCollection.Vertex.Get().GetConstArray();
+			TetCollection.ComponentSpaceVertices(Vertex);
+
+			TArray<bool> VertAdded;
+			VertAdded.Init(false, Vertex.Num());
+			auto DoSkipBoneIndex = [&StrArray, &Skeleton](int32 BoneIndex)
+			{
+				bool Skip = false;
+				FString BoneName = Skeleton->GetReferenceSkeleton().GetBoneName(BoneIndex).ToString();
+				for (FString Elem : StrArray)
+				{
+					if (BoneName.Contains(Elem))
+					{
+						Skip = true;
+						break;
+					}
+				}
+				return Skip;
+			};
+
+			for (int32 b = 0; b < Skeleton->GetReferenceSkeleton().GetNum(); ++b)
+			{
+				FVector3f BonePosition(ComponentPose[b].GetTranslation());
+				int32 ParentIndex=Skeleton->GetReferenceSkeleton().GetParentIndex(b);
+				if (!(ParentIndex == INDEX_NONE || DoSkipBoneIndex(b) || DoSkipBoneIndex(ParentIndex)))
+				{
+					FVector3f ParentPosition(ComponentPose[ParentIndex].GetTranslation());
+					FVector3f RayDir = ParentPosition - BonePosition;
+					Chaos::FReal Length = RayDir.Length();
+					RayDir.Normalize();
+
+					if (Length > Chaos::FReal(1e-8)) 
+					{
+						TSet<int32> BoneVertSet;
+						for (int32 t = 0; t < NumTets; t++) 
+						{
+							int32 i = TetCollection.Tetrahedron[t][0];
+							int32 j = TetCollection.Tetrahedron[t][1];
+							int32 k = TetCollection.Tetrahedron[t][2];
+							int32 l = TetCollection.Tetrahedron[t][3];
+
+							TArray<Chaos::TVec3<Chaos::FRealSingle>> InVertices;
+							InVertices.SetNum(4);
+							InVertices[0][0] = Vertex[i].X; InVertices[0][1] = Vertex[i].Y; InVertices[0][2] = Vertex[i].Z;
+							InVertices[1][0] = Vertex[j].X; InVertices[1][1] = Vertex[j].Y; InVertices[1][2] = Vertex[j].Z;
+							InVertices[2][0] = Vertex[k].X; InVertices[2][1] = Vertex[k].Y; InVertices[2][2] = Vertex[k].Z;
+							InVertices[3][0] = Vertex[l].X; InVertices[3][1] = Vertex[l].Y; InVertices[3][2] = Vertex[l].Z;
+							Chaos::FConvex ConvexTet(InVertices, Chaos::FReal(0));
+							Chaos::FReal OutTime;
+							Chaos::FVec3 OutPosition, OutNormal;
+							int32 OutFaceIndex;
+							bool KeepTet = ConvexTet.Raycast(BonePosition, RayDir, Length, Chaos::FReal(0), OutTime, OutPosition, OutNormal, OutFaceIndex);
+							if (KeepTet) 
+							{	
+								for (int32 c = 0; c < 4; ++c)
+								{	
+									if (!VertAdded[TetCollection.Tetrahedron[t][c]])
+									{
+										VertAdded[TetCollection.Tetrahedron[t][c]] = true;
+										BoneVertSet.Add(TetCollection.Tetrahedron[t][c]);
+									}
+								}
+							}
+						}
+
+						TArray<int32> BoundVerts = BoneVertSet.Array();
+						TArray<float> BoundWeights;
+						BoundWeights.Init(float(1), BoundVerts.Num());
+						if (BoundVerts.Num())
+						{
+							//get local coords of bound verts
+							typedef GeometryCollection::Facades::FKinematicBindingFacade FKinematics;
+							FKinematics Kinematics(InCollection); Kinematics.DefineSchema();
+							if (Kinematics.IsValid())
+							{
+								FKinematics::FBindingKey Binding = Kinematics.SetBoneBindings(ParentIndex, BoundVerts, BoundWeights);
+								TManagedArray<TArray<FVector3f>>& LocalPos = InCollection.AddAttribute<TArray<FVector3f>>("LocalPosition", Binding.GroupName);
+								Kinematics.AddKinematicBinding(Binding);
+								auto DoubleVert = [](FVector3f V) { return FVector3d(V.X, V.Y, V.Z); };
+								auto FloatVert = [](FVector3d V) { return FVector3f(V.X, V.Y, V.Z); };
+								LocalPos[Binding.Index].SetNum(BoundVerts.Num());
+								for (int32 i = 0; i < BoundVerts.Num(); i++)
+								{
+									FVector3f Temp = Vertex[BoundVerts[i]];
+									LocalPos[Binding.Index][i] = FloatVert(ComponentPose[ParentIndex].InverseTransformPosition(DoubleVert(Temp)));
+								}
+							}
+						}
+					}
+				}
+			}
+			GeometryCollection::Facades::FVertexBoneWeightsFacade(InCollection).AddBoneWeightsFromKinematicBindings();
+		}
+		SetValue(Context, MoveTemp(InCollection), &Collection);
+	}
+}
diff --git a/Engine/Plugins/Experimental/ChaosFlesh/Source/ChaosFleshNodes/Private/Dataflow/ChaosFleshKinematicTetrahedralConstraintNode.cpp b/Engine/Plugins/Experimental/ChaosFlesh/Source/ChaosFleshNodes/Private/Dataflow/ChaosFleshKinematicTetrahedralConstraintNode.cpp
deleted file mode 100644
diff --git a/Engine/Plugins/Experimental/ChaosFlesh/Source/ChaosFleshNodes/Private/Dataflow/ChaosFleshNodesPlugin.cpp b/Engine/Plugins/Experimental/ChaosFlesh/Source/ChaosFleshNodes/Private/Dataflow/ChaosFleshNodesPlugin.cpp
--- a/Engine/Plugins/Experimental/ChaosFlesh/Source/ChaosFleshNodes/Private/Dataflow/ChaosFleshNodesPlugin.cpp
+++ b/Engine/Plugins/Experimental/ChaosFlesh/Source/ChaosFleshNodes/Private/Dataflow/ChaosFleshNodesPlugin.cpp
@@ -20,7 +20,7 @@
 #include "Dataflow/ChaosFleshIsolateComponentNode.h"
 #include "Dataflow/ChaosFleshKinematicConstraintNode.h"
 #include "Dataflow/ChaosFleshKinematicMuscleAttachmentsNode.h"
-#include "Dataflow/ChaosFleshKinematicTetrahedralConstraintNode.h"
+#include "Dataflow/ChaosFleshKinematicSkeletonConstraintNode.h"
 #include "Dataflow/ChaosFleshRadialTetrahedronNode.h"
 #include "Dataflow/ChaosFleshTriangleMeshSimulationPropertiesNode.h"
 #include "Dataflow/ChaosFleshSkeletalMeshConstraintNode.h"
@@ -59,8 +59,8 @@ void IChaosFleshNodesPlugin::StartupModule()
 	DATAFLOW_NODE_REGISTER_CREATION_FACTORY(FKinematicBodySetupInitializationDataflowNode);
 	DATAFLOW_NODE_REGISTER_CREATION_FACTORY(FKinematicInitializationDataflowNode);
 	DATAFLOW_NODE_REGISTER_CREATION_FACTORY(FKinematicMuscleAttachmentsDataflowNode);
+	DATAFLOW_NODE_REGISTER_CREATION_FACTORY(FKinematicSkeletonConstraintDataflowNode);
 	DATAFLOW_NODE_REGISTER_CREATION_FACTORY(FKinematicSkeletalMeshInitializationDataflowNode);
-	DATAFLOW_NODE_REGISTER_CREATION_FACTORY(FKinematicTetrahedralBindingsDataflowNode);
 	DATAFLOW_NODE_REGISTER_CREATION_FACTORY(FRadialTetrahedronDataflowNodes);
 	DATAFLOW_NODE_REGISTER_CREATION_FACTORY(FSetFleshBonePositionTargetBindingDataflowNode);
 	DATAFLOW_NODE_REGISTER_CREATION_FACTORY(FSetFleshDefaultPropertiesNode);
diff --git a/Engine/Plugins/Experimental/ChaosFlesh/Source/ChaosFleshNodes/Public/Dataflow/ChaosFleshKinematicSkeletonConstraintNode.h b/Engine/Plugins/Experimental/ChaosFlesh/Source/ChaosFleshNodes/Public/Dataflow/ChaosFleshKinematicSkeletonConstraintNode.h
new file mode 100644
--- /dev/null
+++ b/Engine/Plugins/Experimental/ChaosFlesh/Source/ChaosFleshNodes/Public/Dataflow/ChaosFleshKinematicSkeletonConstraintNode.h
@@ -0,0 +1,45 @@
+// Copyright Epic Games, Inc. All Rights Reserved.
+#pragma once 
+
+#include "CoreMinimal.h"
+#include "Dataflow/DataflowCore.h"
+#include "Dataflow/DataflowEngine.h"
+
+#include "ChaosFleshKinematicSkeletonConstraintNode.generated.h"
+
+class USkeleton;
+
+/* Bind a animation driven skeleton hierarchy into the tetrahedron on the collection. */
+USTRUCT(meta = (DataflowFlesh))
+struct FKinematicSkeletonConstraintDataflowNode : public FDataflowNode
+{
+	GENERATED_USTRUCT_BODY()
+	DATAFLOW_NODE_DEFINE_INTERNAL(FKinematicSkeletonConstraintDataflowNode, "KinematicSkeletonConstraint", "Flesh", "")
+	DATAFLOW_NODE_RENDER_TYPE(FGeometryCollection::StaticType(), "Collection")
+
+public:
+	typedef FManagedArrayCollection DataType;
+
+	/* Pass through collection to place constraints in to. */
+	UPROPERTY(meta = (DataflowInput, DataflowOutput, DisplayName = "Collection", DataflowPassthrough = "Collection"))
+	FManagedArrayCollection Collection;
+
+	/* Skeleton to constraint to the tetrahedron (Must be co-located with the tetrahedron). */
+	UPROPERTY(meta = (DataflowInput, DisplayName = "Skeleton"))
+	TObjectPtr<USkeleton> SkeletonIn = nullptr;
+
+	/* Skeleton bones to exclude from the constraint. */
+	UPROPERTY(EditAnywhere, Category = "Dataflow")
+	FString ExclusionList = "";
+
+	FKinematicSkeletonConstraintDataflowNode(const Dataflow::FNodeParameters& InParam, FGuid InGuid = FGuid::NewGuid())
+		: FDataflowNode(InParam, InGuid)
+	{
+		RegisterInputConnection(&Collection);
+		RegisterOutputConnection(&Collection, &Collection);
+		RegisterInputConnection(&SkeletonIn);
+	}
+
+	virtual void Evaluate(Dataflow::FContext& Context, const FDataflowOutput* Out) const override;
+};
+
diff --git a/Engine/Plugins/Experimental/ChaosFlesh/Source/ChaosFleshNodes/Public/Dataflow/ChaosFleshKinematicTetrahedralConstraintNode.h b/Engine/Plugins/Experimental/ChaosFlesh/Source/ChaosFleshNodes/Public/Dataflow/ChaosFleshKinematicTetrahedralConstraintNode.h
deleted file mode 100644
