ChaosFlesh : Fix for CIS

diff --git a/Engine/Plugins/Experimental/ChaosFlesh/Source/ChaosFlesh/Private/ChaosFlesh/ChaosFleshCollectionFacade.cpp b/Engine/Plugins/Experimental/ChaosFlesh/Source/ChaosFlesh/Private/ChaosFlesh/ChaosFleshCollectionFacade.cpp
--- a/Engine/Plugins/Experimental/ChaosFlesh/Source/ChaosFlesh/Private/ChaosFlesh/ChaosFleshCollectionFacade.cpp
+++ b/Engine/Plugins/Experimental/ChaosFlesh/Source/ChaosFlesh/Private/ChaosFlesh/ChaosFleshCollectionFacade.cpp
@@ -88,5 +88,6 @@ namespace Chaos
 			}
 		}
 	}
-
 }
+
+#undef LOCTEXT_NAMESPACE
