ChaosFlesh : Import transform space update.

diff --git a/Engine/Plugins/Experimental/ChaosFlesh/Source/ChaosFlesh/Private/ChaosFlesh/ChaosFleshCollectionFacade.cpp b/Engine/Plugins/Experimental/ChaosFlesh/Source/ChaosFlesh/Private/ChaosFlesh/ChaosFleshCollectionFacade.cpp
--- a/Engine/Plugins/Experimental/ChaosFlesh/Source/ChaosFlesh/Private/ChaosFlesh/ChaosFleshCollectionFacade.cpp
+++ b/Engine/Plugins/Experimental/ChaosFlesh/Source/ChaosFlesh/Private/ChaosFlesh/ChaosFleshCollectionFacade.cpp
@@ -12,45 +12,49 @@
 namespace Chaos 
 {
 
-	FFleshCollectionFacade::FFleshCollectionFacade(FManagedArrayCollection& InCollection)
+	FFleshCollectionFacade::FFleshCollectionFacade(FFleshCollection& InCollection)
 		: ConstCollection(InCollection)
 		, Collection(&InCollection)
 		, BoneName(InCollection, "BoneName", FTransformCollection::TransformGroup)
 		, Transform(InCollection, FTransformCollection::TransformAttribute, FTransformCollection::TransformGroup)
+		, TransformToGeometryIndex(InCollection, "TransformToGeometryIndex", FTransformCollection::TransformGroup)
 		, Parent(InCollection, FTransformCollection::ParentAttribute, FTransformCollection::TransformGroup)
 		, Child(InCollection, FTransformCollection::ChildrenAttribute, FTransformCollection::TransformGroup)
 		, BoneMap(InCollection, FName("BoneMap"), FName("Vertices"))
 		, Vertex(InCollection, FName("Vertex"), FName("Vertices"))
 		, Indices(InCollection, FName("Indices"), FName("Faces"))
 		, Tetrahedron(InCollection, FTetrahedralCollection::TetrahedronAttribute, FTetrahedralCollection::TetrahedralGroup)
+		, GeometryToTransformIndex(InCollection, "TransformIndex", FGeometryCollection::GeometryGroup)
 		, VertexStart(InCollection, "VertexStart", FGeometryCollection::GeometryGroup)
 		, VertexCount(InCollection, "VertexCount", FGeometryCollection::GeometryGroup)
 		, FaceStart(InCollection, "FaceStart", FGeometryCollection::GeometryGroup)
 		, FaceCount(InCollection, "FaceCount", FGeometryCollection::GeometryGroup)
 	{}
 
-	FFleshCollectionFacade::FFleshCollectionFacade(const FManagedArrayCollection& InCollection)
+	FFleshCollectionFacade::FFleshCollectionFacade(const FFleshCollection& InCollection)
 		: ConstCollection(InCollection)
 		, Collection(nullptr)
 		, BoneName(InCollection, "BoneName", FTransformCollection::TransformGroup)
 		, Transform(InCollection, FTransformCollection::TransformAttribute, FTransformCollection::TransformGroup)
+		, TransformToGeometryIndex(InCollection, "TransformToGeometryIndex", FTransformCollection::TransformGroup)
 		, Parent(InCollection, FTransformCollection::ParentAttribute, FTransformCollection::TransformGroup)
 		, Child(InCollection, FTransformCollection::ChildrenAttribute, FTransformCollection::TransformGroup)
 		, BoneMap(InCollection, FName("BoneMap"), FName("Vertices"))
 		, Vertex(InCollection, FName("Vertex"), FName("Vertices"))
 		, Indices(InCollection, FName("Indices"), FName("Faces"))
 		, Tetrahedron(InCollection, FTetrahedralCollection::TetrahedronAttribute, FTetrahedralCollection::TetrahedralGroup)
+		, GeometryToTransformIndex(InCollection, "TransformIndex", FGeometryCollection::GeometryGroup)
 		, VertexStart(InCollection, "VertexStart", FGeometryCollection::GeometryGroup)
 		, VertexCount(InCollection, "VertexCount", FGeometryCollection::GeometryGroup)
 		, FaceStart(InCollection, "FaceStart", FGeometryCollection::GeometryGroup)
 		, FaceCount(InCollection, "FaceCount", FGeometryCollection::GeometryGroup)
 	{}
 
-	bool FFleshCollectionFacade::IsCompletelyValid() const {
-		return BoneName.IsValid() && Transform.IsValid() && Parent.IsValid() && Child.IsValid() &&
+	bool FFleshCollectionFacade::IsValid() const {
+		return BoneName.IsValid() && Transform.IsValid() && TransformToGeometryIndex.IsValid() && Parent.IsValid() && Child.IsValid() &&
 			BoneMap.IsValid() && Vertex.IsValid() &&
 			Indices.IsValid() && Tetrahedron.IsValid() &&
-			VertexStart.IsValid() && VertexCount.IsValid() && FaceStart.IsValid() && FaceCount.IsValid();
+			GeometryToTransformIndex.IsValid() && VertexStart.IsValid() && VertexCount.IsValid() && FaceStart.IsValid() && FaceCount.IsValid();
 	}
 
 	bool FFleshCollectionFacade::IsTetrahedronValid() const {
@@ -58,15 +62,81 @@ namespace Chaos
 	}
 
 	bool FFleshCollectionFacade::IsHierarchyValid() const {
-		return Transform.IsValid() && Parent.IsValid();
+		return BoneName.IsValid() && Transform.IsValid() && Parent.IsValid() && Child.IsValid();
 	}
 
-	void FFleshCollectionFacade::ComponentSpaceVertices(TArray<FVector3f>& OutComponentSpaceVertices)
+	bool FFleshCollectionFacade::IsGeometryValid() const {
+		return TransformToGeometryIndex.IsValid() && GeometryToTransformIndex.IsValid()
+			&& VertexStart.IsValid() && VertexCount.IsValid()
+			&& FaceStart.IsValid() && FaceCount.IsValid();
+	}
+
+	int FFleshCollectionFacade::NumGeometry() const {
+		if (VertexStart.IsValid())
+		{
+			return VertexStart.Num();
+		}
+		return 0;
+	}
+
+	int FFleshCollectionFacade::NumVertices() const {
+		if (Vertex.IsValid())
+		{
+			return Vertex.Num();
+		}
+		return 0;
+	}
+
+	int FFleshCollectionFacade::NumFaces() const {
+		if (Indices.IsValid())
+		{
+			return Indices.Num();
+		}
+		return 0;
+	}
+
+	int FFleshCollectionFacade::AppendGeometry(const FFleshCollection& NewGeomerty) 
+	{
+		int32 GeomIndex = NumGeometry();
+		Collection->AppendGeometry(NewGeomerty);
+		return GeomIndex;
+	}
+
+	void FFleshCollectionFacade::GlobalMatrices(TArray<FTransform>& OutComponentTransform)
+	{
+		if (Transform.IsValid() && Parent.IsValid())
+		{
+			GeometryCollectionAlgo::GlobalMatrices(Transform.Get(), Parent.Get(), OutComponentTransform);
+		}
+	}
+
+	FTransform3f FFleshCollectionFacade::GlobalMatrix3f(int32 InIndex)
+	{
+		if (Transform.IsValid() && Parent.IsValid())
+		{
+			if (0 <= InIndex && InIndex < Transform.Num())
+			{
+				return GeometryCollectionAlgo::GlobalMatrix3f(Transform.Get(), Parent.Get(), InIndex);
+			}
+		}
+		return FTransform3f::Identity;
+	}
+
+
+	int FFleshCollectionFacade::NumTransforms() const {
+		if (Transform.IsValid())
+		{
+			return Transform.Num();
+		}
+		return 0;
+	}
+
+	void FFleshCollectionFacade::ComponentSpaceVertices(TArray<FVector3f>& OutComponentSpaceVertices) const
 	{
 		ComponentSpaceVertices(OutComponentSpaceVertices, 0, Vertex.Num());
 	}
 
-	void FFleshCollectionFacade::ComponentSpaceVertices(TArray<FVector3f>& OutComponentSpaceVertices, int32 Start, int32 Count)
+	void FFleshCollectionFacade::ComponentSpaceVertices(TArray<FVector3f>& OutComponentSpaceVertices, int32 Start, int32 Count) const
 	{
 		if (IsHierarchyValid() && Vertex.IsValid())
 		{
diff --git a/Engine/Plugins/Experimental/ChaosFlesh/Source/ChaosFlesh/Public/ChaosFlesh/ChaosFleshCollectionFacade.h b/Engine/Plugins/Experimental/ChaosFlesh/Source/ChaosFlesh/Public/ChaosFlesh/ChaosFleshCollectionFacade.h
--- a/Engine/Plugins/Experimental/ChaosFlesh/Source/ChaosFlesh/Public/ChaosFlesh/ChaosFleshCollectionFacade.h
+++ b/Engine/Plugins/Experimental/ChaosFlesh/Source/ChaosFlesh/Public/ChaosFlesh/ChaosFleshCollectionFacade.h
@@ -1,54 +1,81 @@
 // Copyright Epic Games, Inc. All Rights Reserved.
+#pragma once
 
 #include "CoreMinimal.h"
-#include "GeometryCollection/ManagedArrayCollection.h"
+#include "ChaosFlesh/FleshCollection.h"
+#include "GeometryCollection/ManagedArray.h"
 #include "GeometryCollection/ManagedArrayAccessor.h"
 
-
 namespace Chaos {
 
 	class CHAOSFLESH_API FFleshCollectionFacade
 	{
-		const FManagedArrayCollection& ConstCollection;
-		FManagedArrayCollection* Collection = nullptr;
+		const FFleshCollection& ConstCollection;
+		FFleshCollection* Collection = nullptr;
 	public:
 
-		FFleshCollectionFacade(FManagedArrayCollection& InCollection);
+		FFleshCollectionFacade(FFleshCollection& InCollection);
 
-		FFleshCollectionFacade(const FManagedArrayCollection& InCollection);
+		FFleshCollectionFacade(const FFleshCollection& InCollection);
 		
 		/*Are all the public attributes avaiable.*/
-		bool IsCompletelyValid() const;
+		bool IsValid() const;
 	
 		/*Has tetrahedral attributes {Tetrahderon, Vertices}.*/
 		bool IsTetrahedronValid() const;
 
-		/*Has hierarchy attributes {Transform, Parent}.*/
+		/*Has hierarchy attributes {BoneName, Transform, Parent, Child}.*/
 		bool IsHierarchyValid() const;
 
+		/*Has geometry attributes {TransformToGeometryIndex VertexStart VertexCount FaceStart FaceCount}.*/
+		bool IsGeometryValid() const;
+
+		/* Num of transform elements in the collection */
+		int NumTransforms() const;
+
+		/* Num of vertex elements in the collection */
+		int NumVertices() const;
+
+		/* Num of face elements in the collection */
+		int NumFaces() const;
+
+		/* Num of geometry elements in the collection */
+		int NumGeometry() const;
+
+		/* Add geometry to the collection*/
+		int AppendGeometry(const FFleshCollection& NewGeomerty);
+
+		/* Global Matrices of the collection */
+		void GlobalMatrices(TArray<FTransform>& ComponentTransform);
+
+		/* Single Global Matrix of index in the collection */
+		FTransform3f GlobalMatrix3f(int32 InIndex);
+
 		/*View a Attribute on the collection*/
 		template<class T>
 		const TManagedArray<T>* FindAttribute(FString AttributeName, FString Group) const;
 
 		/*Edit a Attribute on the collection*/
 		template<class T>
-		TManagedArray<T>* ModifyAttribute(FString AttributeName, FString Group);
+		TManagedArray<T>* ModifyAttribute(FString AttributeName, FString Group) const;
 
 		/* All the vertices mapped into component space. */
-		void ComponentSpaceVertices(TArray<FVector3f>& OutComponentSpaceVertices);
+		void ComponentSpaceVertices(TArray<FVector3f>& OutComponentSpaceVertices) const;
 
 		/* Range of vertices mapped into component space*/
-		void ComponentSpaceVertices(TArray<FVector3f>& OutComponentSpaceVertices, int32 Start, int32 Count);
+		void ComponentSpaceVertices(TArray<FVector3f>& OutComponentSpaceVertices, int32 Start, int32 Count) const;
 
 		/*Public Attributes*/
 		TManagedArrayAccessor<FString> BoneName;
 		TManagedArrayAccessor<FTransform3f> Transform;
+		TManagedArrayAccessor<int32> TransformToGeometryIndex;
 		TManagedArrayAccessor<int32> Parent;
 		TManagedArrayAccessor< TSet<int32> > Child;
 		TManagedArrayAccessor<int32> BoneMap;
 		TManagedArrayAccessor<FVector3f> Vertex;
 		TManagedArrayAccessor<FIntVector3> Indices;
 		TManagedArrayAccessor<FIntVector4> Tetrahedron;
+		TManagedArrayAccessor<int32> GeometryToTransformIndex;
 		TManagedArrayAccessor<int32> VertexStart;
 		TManagedArrayAccessor<int32> VertexCount;
 		TManagedArrayAccessor<int32> FaceStart;
diff --git a/Engine/Plugins/Experimental/ChaosFlesh/Source/ChaosFleshEngine/Private/ChaosFlesh/ChaosDeformableTetrahedralComponent.cpp b/Engine/Plugins/Experimental/ChaosFlesh/Source/ChaosFleshEngine/Private/ChaosFlesh/ChaosDeformableTetrahedralComponent.cpp
--- a/Engine/Plugins/Experimental/ChaosFlesh/Source/ChaosFleshEngine/Private/ChaosFlesh/ChaosDeformableTetrahedralComponent.cpp
+++ b/Engine/Plugins/Experimental/ChaosFlesh/Source/ChaosFleshEngine/Private/ChaosFlesh/ChaosDeformableTetrahedralComponent.cpp
@@ -11,6 +11,7 @@
 #include "ChaosFlesh/ChaosDeformableSolverActor.h"
 #include "ChaosFlesh/ChaosDeformableSolverComponent.h"
 #include "ChaosFlesh/ChaosDeformableTypes.h"
+#include "ChaosFlesh/ChaosFleshCollectionFacade.h"
 #include "ChaosFlesh/FleshDynamicAsset.h"
 #include "Components/SkeletalMeshComponent.h"
 #include "Dataflow/DataflowEngineUtil.h"
@@ -458,10 +459,12 @@ void UDeformableTetrahedralComponent::RenderProceduralMesh()
 				{
 					PERF_SCOPE(STAT_ChaosDeformable_UDeformableTetrahedralComponent_RenderProceduralMesh);
 
-					if (const FFleshCollection* Flesh = FleshAsset->GetCollection())
+					if (const FFleshCollection* FleshCollection = FleshAsset->GetCollection())
 					{
-						int32 NumVertices = Flesh->NumElements(FGeometryCollection::VerticesGroup);
-						int32 NumFaces = Flesh->NumElements(FGeometryCollection::FacesGroup);
+						const Chaos::FFleshCollectionFacade Flesh(*FleshCollection);
+
+						int32 NumVertices = Flesh.NumVertices();
+						int32 NumFaces = Flesh.NumFaces();
 						if (NumFaces && NumVertices)
 						{
 							if (RenderMesh && RenderMesh->Vertices.Num() != NumFaces * 3)
@@ -472,12 +475,14 @@ void UDeformableTetrahedralComponent::RenderProceduralMesh()
 							if (!RenderMesh)
 							{
 								RenderMesh = new FFleshRenderMesh;
+								TArray<FVector3f> Vertex;
+								Flesh.ComponentSpaceVertices(Vertex);
 
 								for (int i = 0; i < NumFaces; ++i)
 								{
-									const auto& P1 = Flesh->Vertex[Flesh->Indices[i][0]];
-									const auto& P2 = Flesh->Vertex[Flesh->Indices[i][1]];
-									const auto& P3 = Flesh->Vertex[Flesh->Indices[i][2]];
+									const auto& P1 = Vertex[Flesh.Indices[i][0]];
+									const auto& P2 = Vertex[Flesh.Indices[i][1]];
+									const auto& P3 = Vertex[Flesh.Indices[i][2]];
 
 									RenderMesh->Vertices.Add(FVector(P1));
 									RenderMesh->Vertices.Add(FVector(P2));
@@ -513,21 +518,22 @@ void UDeformableTetrahedralComponent::RenderProceduralMesh()
 							}
 							else
 							{
+								TArray<FVector3f> RenderVertex; 
+								Flesh.ComponentSpaceVertices(RenderVertex);
 
-								const TManagedArray<FVector3f>* RenderVertex = &Flesh->Vertex;
 								if (GetDynamicCollection())
 								{
 									const TManagedArray<FVector3f>& DynamicVertex = GetDynamicCollection()->GetPositions();
-									if (DynamicVertex.Num()) RenderVertex = &DynamicVertex;
+									if (DynamicVertex.Num()) RenderVertex = DynamicVertex.GetConstArray();
 								}
 								auto InRange = [](int32 Size, int32 Val) { return 0 <= Val && Val < Size; };
 
 								// Display only
 								for (int i = 0; i < NumFaces; ++i)
 								{
-									const auto& P1 = (*RenderVertex)[Flesh->Indices[i][0]];
-									const auto& P2 = (*RenderVertex)[Flesh->Indices[i][1]];
-									const auto& P3 = (*RenderVertex)[Flesh->Indices[i][2]];
+									const auto& P1 = RenderVertex[Flesh.Indices[i][0]];
+									const auto& P2 = RenderVertex[Flesh.Indices[i][1]];
+									const auto& P3 = RenderVertex[Flesh.Indices[i][2]];
 
 									RenderMesh->Vertices[3 * i] = FVector(P1);
 									RenderMesh->Vertices[3 * i + 1] = FVector(P2);
diff --git a/Engine/Plugins/Experimental/ChaosFlesh/Source/ChaosFleshNodes/Private/Dataflow/ChaosFleshCreateTetrahedronNode.cpp b/Engine/Plugins/Experimental/ChaosFlesh/Source/ChaosFleshNodes/Private/Dataflow/ChaosFleshCreateTetrahedronNode.cpp
--- a/Engine/Plugins/Experimental/ChaosFlesh/Source/ChaosFleshNodes/Private/Dataflow/ChaosFleshCreateTetrahedronNode.cpp
+++ b/Engine/Plugins/Experimental/ChaosFlesh/Source/ChaosFleshNodes/Private/Dataflow/ChaosFleshCreateTetrahedronNode.cpp
@@ -10,6 +10,7 @@
 #include "Chaos/UniformGrid.h"
 #include "ChaosFlesh/FleshCollection.h"
 #include "ChaosFlesh/FleshCollectionUtility.h"
+#include "ChaosFlesh/ChaosFleshCollectionFacade.h"
 #include "ChaosLog.h"
 #include "Dataflow/DataflowInputOutput.h"
 #include "Dataflow/ChaosFleshNodesUtility.h"
@@ -39,53 +40,28 @@ void FCreateTetrahedronDataflowNode::Evaluate(Dataflow::FContext& Context, const
 {
 	if (Out->IsA<DataType>(&Collection))
 	{
-		TUniquePtr<FFleshCollection> InCollection(GetValue<DataType>(Context, &Collection).NewCopy<FFleshCollection>());
-		TUniquePtr<FFleshCollection> InSourceCollection(GetValue<DataType>(Context, &SourceCollection).NewCopy<FFleshCollection>());
+		TUniquePtr<FFleshCollection> InCollectionVal(GetValue<DataType>(Context, &Collection).NewCopy<FFleshCollection>());
+		Chaos::FFleshCollectionFacade Target(*InCollectionVal.Get());
 
-		// TransformGroup
-		int32 NumTransforms = InSourceCollection->NumElements(FTransformCollection::TransformGroup);
-		TManagedArray<FString>* TransformName = InSourceCollection->FindAttribute<FString>("BoneName", FTransformCollection::TransformGroup);
-		TManagedArray<int32>* TransformToGeometryIndex = InSourceCollection->FindAttribute<int32>("TransformToGeometryIndex", FTransformCollection::TransformGroup);
-		const TManagedArray<int32>* Parent = InSourceCollection->FindAttribute<int32>(FTransformCollection::ParentAttribute, FTransformCollection::TransformGroup);
-		const TManagedArray<FTransform3f>* LocalSpaceTransform = InSourceCollection->FindAttribute<FTransform3f>(FTransformCollection::TransformAttribute, FTransformCollection::TransformGroup);
-		// Geometry Group
-		int32 NumGeometry = InSourceCollection->NumElements(FGeometryCollection::GeometryGroup);
-		TManagedArray<int32>* GroupToTransformIndex = InSourceCollection->FindAttribute<int32>("TransformIndex", FGeometryCollection::GeometryGroup);
-		TManagedArray<int32>* VertexCount = InSourceCollection->FindAttribute<int32>("VertexCount", FGeometryCollection::GeometryGroup);
-		TManagedArray<int32>* VertexStart = InSourceCollection->FindAttribute<int32>("VertexStart", FGeometryCollection::GeometryGroup);
-		TManagedArray<int32>* FaceCount = InSourceCollection->FindAttribute<int32>("FaceCount", FGeometryCollection::GeometryGroup);
-		TManagedArray<int32>* FaceStart = InSourceCollection->FindAttribute<int32>("FaceStart", FGeometryCollection::GeometryGroup);
+		TUniquePtr<FFleshCollection> InSourceCollectionVal(GetValue<DataType>(Context, &SourceCollection).NewCopy<FFleshCollection>());
+		Chaos::FFleshCollectionFacade Source(*InSourceCollectionVal.Get());
 
-		// Vertices Group
-		int32 NumVertices = InSourceCollection->NumElements(FGeometryCollection::VerticesGroup);
-		const TManagedArray<FVector3f>* Vertex = InSourceCollection->FindAttribute<FVector3f>("Vertex", FGeometryCollection::VerticesGroup);
-		// Faces Group
-		int32 NumTriangles = InSourceCollection->NumElements(FGeometryCollection::FacesGroup);
-		const TManagedArray<FIntVector>* Faces = InSourceCollection->FindAttribute<FIntVector>("Indices", FGeometryCollection::FacesGroup);
-
-		if (TransformName && TransformToGeometryIndex && Parent && LocalSpaceTransform &&
-			GroupToTransformIndex && VertexCount && VertexStart && FaceCount && FaceStart
-			&& Vertex && Faces)
+		if (Source.IsValid() && Target.IsValid())
 		{
-			auto GetParentIndex = [&InCollection, &InSourceCollection](int32 Gdx)
+			auto GetParentIndex = [&Target, &Source](int32 Gdx)
 			{
 				if (Gdx >= 0)
 				{
-					TManagedArray<FString>* SourceTransformName = InSourceCollection->FindAttribute<FString>("BoneName", FTransformCollection::TransformGroup);
-					TManagedArray<FString>* TransformName = InCollection->FindAttribute<FString>("BoneName", FTransformCollection::TransformGroup);
-					TManagedArray<int32>* SourceGroupToTransformIndex = InSourceCollection->FindAttribute<int32>("TransformIndex", FGeometryCollection::GeometryGroup);
-					TManagedArray<int32>* GroupToTransformIndex = InCollection->FindAttribute<int32>("TransformIndex", FGeometryCollection::GeometryGroup);
-					TManagedArray<int32>* ParentIndex = InCollection->FindAttribute<int32>("Parent", FGeometryCollection::TransformGroup);
-					if (SourceTransformName && TransformName && SourceGroupToTransformIndex && GroupToTransformIndex && ParentIndex)
+					if (Target.IsHierarchyValid())
 					{
-						if (Gdx < SourceGroupToTransformIndex->Num() && Gdx < GroupToTransformIndex->Num())
+						if (Gdx < Source.GeometryToTransformIndex.Num() && Gdx < Target.GeometryToTransformIndex.Num())
 						{
-							int32 SrcTdx = (*SourceGroupToTransformIndex)[Gdx], Tdx = (*GroupToTransformIndex)[Gdx];
-							if (0 <= SrcTdx && SrcTdx < SourceTransformName->Num() && 0 <= Tdx && Tdx < TransformName->Num())
+							int32 SrcTdx = Source.GeometryToTransformIndex[Gdx], Tdx = Target.GeometryToTransformIndex[Gdx];
+							if (0 <= SrcTdx && SrcTdx < Source.BoneName.Num() && 0 <= Tdx && Tdx < Target.BoneName.Num())
 							{
-								if ((*SourceTransformName)[SrcTdx].Equals((*TransformName)[Tdx]))
+								if (Source.BoneName[SrcTdx].Equals(Target.BoneName[Tdx]))
 								{
-									return (*ParentIndex)[Tdx];
+									return Target.Parent[Tdx];
 								}
 							}
 						}
@@ -94,10 +70,10 @@ void FCreateTetrahedronDataflowNode::Evaluate(Dataflow::FContext& Context, const
 				return (int32)INDEX_NONE;
 			};
 
-			TArray<int32> ProcessGeometryIndices = Dataflow::GetMatchingMeshIndices(Selection, InSourceCollection.Get());
+			TArray<int32> ProcessGeometryIndices = Dataflow::GetMatchingMeshIndices(Selection, InSourceCollectionVal.Get());
 
 			TArray<TUniquePtr<FFleshCollection>> CollectionBuffer;
-			for (int32 Gdx = 0; Gdx < NumGeometry; Gdx++)
+			for (int32 Gdx = 0; Gdx < Source.NumGeometry(); Gdx++)
 			{
 				if (ProcessGeometryIndices.Contains(Gdx))
 				{
@@ -108,19 +84,19 @@ void FCreateTetrahedronDataflowNode::Evaluate(Dataflow::FContext& Context, const
 			}
 
 			TArray<FTransform> ComponentTransform;
-			GeometryCollectionAlgo::GlobalMatrices(*LocalSpaceTransform, *Parent, ComponentTransform);
+			Source.GlobalMatrices(ComponentTransform);
 
 			ParallelFor(ProcessGeometryIndices.Num(), [&](int32 i)
 			{
 				int32 Gdx = ProcessGeometryIndices[i];
-			    int32 Tdx = (*GroupToTransformIndex)[Gdx];
+			    int32 Tdx = Source.GeometryToTransformIndex[Gdx];
 				FFleshCollection& TetCollection = *CollectionBuffer[Gdx];
 
 				UE::Geometry::FDynamicMesh3 DynamicMesh;
-				int32 vStart = (*VertexStart)[Gdx], vEnd = vStart + (*VertexCount)[Gdx];
-				for (int Vdx = vStart; Vdx < vEnd; Vdx++) DynamicMesh.AppendVertex(ComponentTransform[Tdx].TransformPosition(FVector((*Vertex)[Vdx])));
-				int32 fStart = (*FaceStart)[Gdx], fEnd = fStart + (*FaceCount)[Gdx];
-				for (int Fdx = fStart; Fdx < fEnd; Fdx++) DynamicMesh.AppendTriangle((*Faces)[Fdx] - FIntVector(vStart));
+				int32 vStart = Source.VertexStart[Gdx], vEnd = vStart + Source.VertexCount[Gdx];
+				for (int Vdx = vStart; Vdx < vEnd; Vdx++) DynamicMesh.AppendVertex(ComponentTransform[Tdx].TransformPosition(FVector(Source.Vertex[Vdx])));
+				int32 fStart = Source.FaceStart[Gdx], fEnd = fStart + Source.FaceCount[Gdx];
+				for (int Fdx = fStart; Fdx < fEnd; Fdx++) DynamicMesh.AppendTriangle(Source.Indices[Fdx] - FIntVector(vStart));
 				DynamicMesh.CompactInPlace();
 
 				if (Method == TetMeshingMethod::IsoStuffing)
@@ -141,90 +117,72 @@ void FCreateTetrahedronDataflowNode::Evaluate(Dataflow::FContext& Context, const
 				}
 			});
 
-			auto AppendProcessedGeometry = [&ProcessGeometryIndices, &CollectionBuffer, &InSourceCollection](FFleshCollection& ToCollection)
+			auto AppendProcessedGeometry = [&ProcessGeometryIndices, &CollectionBuffer, &Source, &Target]()
 			{
 				// find the index in the to transform group based on name
-				auto FindParentIndexFromName = [](FFleshCollection& fpColllection, FString fpName) 
+				auto FindParentIndexFromName = [&Target](FString fpName)
 				{
 					if (!fpName.IsEmpty())
 					{
-						const TManagedArray<FString>& Names = fpColllection.GetAttribute<FString>("BoneName", FTransformCollection::TransformGroup);
-						for (int32 i = 0; i < Names.Num(); i++) if (Names[i].Equals(fpName)) return i;
+						for (int32 i = 0; i < Target.BoneName.Num(); i++) if (Target.BoneName[i].Equals(fpName)) return i;
 					}
 					return (int32)INDEX_NONE;
 				};
 
-				TManagedArray<int32>* SourceGroupToTransformIndex = InSourceCollection->FindAttribute<int32>("TransformIndex", FGeometryCollection::GeometryGroup);
-				TManagedArray<FString>* SourceTransformName = InSourceCollection->FindAttribute<FString>("BoneName", FTransformCollection::TransformGroup);
-				TManagedArray<int32>* SourceParentIndex = InSourceCollection->FindAttribute<int32>("Parent", FTransformCollection::TransformGroup);
-				TManagedArray<int32>* ToGroupToTransformIndex = ToCollection.FindAttribute<int32>("TransformIndex", FGeometryCollection::GeometryGroup);
-				TManagedArray<int32>* TransformToGroupIndex = ToCollection.FindAttribute<int32>("TransformToGeometryIndex", FTransformCollection::TransformGroup);
-				TManagedArray<FString>* ToTransformName = ToCollection.FindAttribute<FString>("BoneName", FTransformCollection::TransformGroup);
-				TManagedArray<FTransform3f>* ToTransform = ToCollection.FindAttribute<FTransform3f>("Transform", FTransformCollection::TransformGroup);
-				TManagedArray<int32>* ToParentIndex = ToCollection.FindAttribute<int32>("Parent", FTransformCollection::TransformGroup);
-				TManagedArray<TSet<int32> >* ToChildIndex = ToCollection.FindAttribute< TSet<int32> >(FTransformCollection::ChildrenAttribute, FTransformCollection::TransformGroup);
-				TManagedArray<FVector3f>* ToVertex = ToCollection.FindAttribute<FVector3f>("Vertex", FGeometryCollection::VerticesGroup);
-				TManagedArray<int32>* ToVertexCount = ToCollection.FindAttribute<int32>("VertexCount", FGeometryCollection::GeometryGroup);
-				TManagedArray<int32>* ToVertexStart = ToCollection.FindAttribute<int32>("VertexStart", FGeometryCollection::GeometryGroup);
-
-				if (SourceGroupToTransformIndex && SourceTransformName && ToGroupToTransformIndex && ToTransformName && ToParentIndex && SourceParentIndex && ToChildIndex)
+				for (int32 Sdx = 0; Sdx < ProcessGeometryIndices.Num(); Sdx++)
 				{
-					for (int32 Sdx = 0; Sdx < ProcessGeometryIndices.Num(); Sdx++)
+					int32 Gdx = ProcessGeometryIndices[Sdx];
+					if (CollectionBuffer[Gdx] && CollectionBuffer[Gdx]->NumElements(FGeometryCollection::GeometryGroup))
 					{
-						int32 Gdx = ProcessGeometryIndices[Sdx];
-						if (CollectionBuffer[Gdx] && CollectionBuffer[Gdx]->NumElements(FGeometryCollection::GeometryGroup))
-						{
-							int32 GeomIndex = ToCollection.NumElements(FGeometryCollection::GeometryGroup);
-							ToCollection.AppendGeometry(*CollectionBuffer[Gdx]);
-							// source data
-							int32 SourceNumTransforms = InSourceCollection->NumElements(FGeometryCollection::TransformGroup);
-							int32 SourceTransformIndex = (*SourceGroupToTransformIndex)[Gdx];
-							int32 SourceTransformParent = (0 <= SourceTransformIndex && SourceTransformIndex < SourceNumTransforms) ? (*SourceParentIndex)[SourceTransformIndex] : INDEX_NONE;
-							FString SourceParentName = (0 <= SourceTransformParent && SourceTransformParent < SourceNumTransforms) ? (*SourceTransformName)[SourceTransformParent] : FString("");
-							// target data
-							int32 ToNumTransforms = ToCollection.NumElements(FGeometryCollection::TransformGroup);
-							int32 ToGeomTransformIndex = (*ToGroupToTransformIndex)[GeomIndex];
+						int32 GeomIndex = Target.AppendGeometry(*CollectionBuffer[Gdx]);
+						// source data
+						int32 SourceNumTransforms = Source.NumTransforms();
+						int32 SourceTransformIndex = Source.GeometryToTransformIndex[Gdx];
+						int32 SourceTransformParent = (0 <= SourceTransformIndex && SourceTransformIndex < SourceNumTransforms) ? Source.Parent[SourceTransformIndex] : INDEX_NONE;
+						FString SourceParentName = (0 <= SourceTransformParent && SourceTransformParent < SourceNumTransforms) ? Source.BoneName[SourceTransformParent] : FString("");
+						// target data
+						int32 ToNumTransforms = Target.NumTransforms();
+						int32 ToGeomTransformIndex = Target.GeometryToTransformIndex[GeomIndex];
 
-							FString TetName = FString::Printf(TEXT("Tet%d"), GeomIndex);
-							if (0 <= SourceTransformIndex && SourceTransformIndex < SourceNumTransforms)
+						FString TetName = FString::Printf(TEXT("Tet%d"), GeomIndex);
+						if (0 <= SourceTransformIndex && SourceTransformIndex < SourceNumTransforms)
+						{
+							if (!Source.BoneName[SourceTransformIndex].IsEmpty())
 							{
-								if (!(*SourceTransformName)[SourceTransformIndex].IsEmpty())
-								{
-									TetName = FString::Printf(TEXT("%s_%s"), *(*SourceTransformName)[SourceTransformIndex], *TetName);
-								}
+								TetName = FString::Printf(TEXT("%s_%s"), *Source.BoneName[SourceTransformIndex], *TetName);
 							}
-							if (0 <= ToGeomTransformIndex && ToGeomTransformIndex < ToNumTransforms)
+						}
+						if (0 <= ToGeomTransformIndex && ToGeomTransformIndex < ToNumTransforms)
+						{
+							// set the name
+							Target.BoneName.ModifyAt(ToGeomTransformIndex,TetName);
+
+							// set transform to geometry and geometry to transform mappings
+							Target.TransformToGeometryIndex.ModifyAt(ToGeomTransformIndex, GeomIndex);
+							Target.GeometryToTransformIndex.ModifyAt(GeomIndex, ToGeomTransformIndex);
+							// set the parent and child mappings
+							Target.Parent.ModifyAt(ToGeomTransformIndex,FindParentIndexFromName(SourceParentName));
+							if (Target.Parent[ToGeomTransformIndex] != INDEX_NONE)
 							{
-								// set the name
-								(*ToTransformName)[ToGeomTransformIndex] = TetName;
-				
-								// set transform to geometry and geometry to transform mappings
-								(*TransformToGroupIndex)[ToGeomTransformIndex] = GeomIndex;
-								(*ToGroupToTransformIndex)[GeomIndex] = ToGeomTransformIndex;
-								// set the parent and child mappings
-								(*ToParentIndex)[ToGeomTransformIndex] = FindParentIndexFromName(ToCollection, SourceParentName);
-								if ((*ToParentIndex)[ToGeomTransformIndex] != INDEX_NONE)
-								{
-									(*ToChildIndex)[(*ToParentIndex)[ToGeomTransformIndex]].Add(ToGeomTransformIndex);
-								}
+								Target.Child.Modify()[Target.Parent[ToGeomTransformIndex]].Add(ToGeomTransformIndex);
+							}
 
-								if (ToGeomTransformIndex != INDEX_NONE)
+							if (ToGeomTransformIndex != INDEX_NONE)
+							{
+								int32 VertexEnd = Target.VertexStart[GeomIndex] + Target.VertexCount[GeomIndex];
+								FTransform3f ParentTransform = Target.GlobalMatrix3f(ToGeomTransformIndex);
+								for (int Vdx = Target.VertexStart[GeomIndex]; Vdx < VertexEnd; Vdx++)
 								{
-									int32 VertexEnd = (*ToVertexStart)[GeomIndex] + (*ToVertexCount)[GeomIndex];
-									FTransform3f ParentTransform = GeometryCollectionAlgo::GlobalMatrix3f(*ToTransform, *ToParentIndex, ToGeomTransformIndex);
-									for (int Vdx = (*ToVertexStart)[GeomIndex]; Vdx < VertexEnd; Vdx++)
-									{
-										(*ToVertex)[Vdx] = ParentTransform.InverseTransformPosition((*ToVertex)[Vdx]);
-									}
+									Target.Vertex.ModifyAt(Vdx,ParentTransform.InverseTransformPosition(Target.Vertex[Vdx]));
 								}
 							}
 						}
 					}
 				}
 			};
-			AppendProcessedGeometry(*InCollection);
+			AppendProcessedGeometry();
 		}
-		SetValue<const DataType&>(Context, *InCollection, &Collection);
+		SetValue<const DataType&>(Context, *InCollectionVal, &Collection);
 	}
 }
 
diff --git a/Engine/Plugins/Experimental/ChaosFlesh/Source/ChaosFleshNodes/Private/Dataflow/ChaosFleshKinematicSkeletonConstraintNode.cpp b/Engine/Plugins/Experimental/ChaosFlesh/Source/ChaosFleshNodes/Private/Dataflow/ChaosFleshKinematicSkeletonConstraintNode.cpp
--- a/Engine/Plugins/Experimental/ChaosFlesh/Source/ChaosFleshNodes/Private/Dataflow/ChaosFleshKinematicSkeletonConstraintNode.cpp
+++ b/Engine/Plugins/Experimental/ChaosFlesh/Source/ChaosFleshNodes/Private/Dataflow/ChaosFleshKinematicSkeletonConstraintNode.cpp
@@ -15,8 +15,9 @@ void FKinematicSkeletonConstraintDataflowNode::Evaluate(Dataflow::FContext& Cont
 {
 	if (Out->IsA<DataType>(&Collection))
 	{
-		DataType InCollection = GetValue<DataType>(Context, &Collection);
-		Chaos::FFleshCollectionFacade TetCollection(InCollection);
+		TUniquePtr<FFleshCollection> InCollection(GetValue<DataType>(Context, &Collection).NewCopy<FFleshCollection>());
+		Chaos::FFleshCollectionFacade TetCollection(*InCollection.Get());
+
 		TObjectPtr<USkeleton> Skeleton = GetValue<TObjectPtr<USkeleton>>(Context, &SkeletonIn);
 		if (Skeleton && TetCollection.IsTetrahedronValid() )
 		{
@@ -100,11 +101,11 @@ void FKinematicSkeletonConstraintDataflowNode::Evaluate(Dataflow::FContext& Cont
 						{
 							//get local coords of bound verts
 							typedef GeometryCollection::Facades::FKinematicBindingFacade FKinematics;
-							FKinematics Kinematics(InCollection); Kinematics.DefineSchema();
+							FKinematics Kinematics(*InCollection.Get()); Kinematics.DefineSchema();
 							if (Kinematics.IsValid())
 							{
 								FKinematics::FBindingKey Binding = Kinematics.SetBoneBindings(ParentIndex, BoundVerts, BoundWeights);
-								TManagedArray<TArray<FVector3f>>& LocalPos = InCollection.AddAttribute<TArray<FVector3f>>("LocalPosition", Binding.GroupName);
+								TManagedArray<TArray<FVector3f>>& LocalPos = InCollection->AddAttribute<TArray<FVector3f>>("LocalPosition", Binding.GroupName);
 								Kinematics.AddKinematicBinding(Binding);
 								auto DoubleVert = [](FVector3f V) { return FVector3d(V.X, V.Y, V.Z); };
 								auto FloatVert = [](FVector3d V) { return FVector3f(V.X, V.Y, V.Z); };
@@ -119,8 +120,9 @@ void FKinematicSkeletonConstraintDataflowNode::Evaluate(Dataflow::FContext& Cont
 					}
 				}
 			}
-			GeometryCollection::Facades::FVertexBoneWeightsFacade(InCollection).AddBoneWeightsFromKinematicBindings();
+			GeometryCollection::Facades::FVertexBoneWeightsFacade(*InCollection.Get()).AddBoneWeightsFromKinematicBindings();
 		}
-		SetValue(Context, MoveTemp(InCollection), &Collection);
+		FManagedArrayCollection& InBaseCollection = *InCollection.Release();
+		SetValue(Context, MoveTemp(InBaseCollection), &Collection);
 	}
 }
