ChaosFlesh : Deprecate some nodes.

diff --git a/Engine/Plugins/Experimental/ChaosFlesh/Source/ChaosFleshDeprecatedNodes/Public/Dataflow/ChaosFleshConstructTetGridNode.h b/Engine/Plugins/Experimental/ChaosFlesh/Source/ChaosFleshDeprecatedNodes/Public/Dataflow/ChaosFleshConstructTetGridNode.h
--- a/Engine/Plugins/Experimental/ChaosFlesh/Source/ChaosFleshDeprecatedNodes/Public/Dataflow/ChaosFleshConstructTetGridNode.h
+++ b/Engine/Plugins/Experimental/ChaosFlesh/Source/ChaosFleshDeprecatedNodes/Public/Dataflow/ChaosFleshConstructTetGridNode.h
@@ -10,7 +10,7 @@
 #include "ChaosFleshConstructTetGridNode.generated.h"
 
 
-USTRUCT(meta = (DataflowFlesh))
+USTRUCT(meta = (DataflowFlesh, Deprecated = "5.4"))
 struct FConstructTetGridNode : public FDataflowNode
 {
 	GENERATED_USTRUCT_BODY()
diff --git a/Engine/Plugins/Experimental/ChaosFlesh/Source/ChaosFleshDeprecatedNodes/Public/Dataflow/ChaosFleshCreateTetrahedralCollectionNode.h b/Engine/Plugins/Experimental/ChaosFlesh/Source/ChaosFleshDeprecatedNodes/Public/Dataflow/ChaosFleshCreateTetrahedralCollectionNode.h
--- a/Engine/Plugins/Experimental/ChaosFlesh/Source/ChaosFleshDeprecatedNodes/Public/Dataflow/ChaosFleshCreateTetrahedralCollectionNode.h
+++ b/Engine/Plugins/Experimental/ChaosFlesh/Source/ChaosFleshDeprecatedNodes/Public/Dataflow/ChaosFleshCreateTetrahedralCollectionNode.h
@@ -21,7 +21,7 @@ namespace UE {
 }
 
 
-USTRUCT(meta = (DataflowFlesh))
+USTRUCT(meta = (DataflowFlesh, Deprecated = "5.4"))
 struct FGenerateTetrahedralCollectionDataflowNodes : public FDataflowNode
 {
 	GENERATED_USTRUCT_BODY()
diff --git a/Engine/Plugins/Experimental/ChaosFlesh/Source/ChaosFleshDeprecatedNodes/Public/Dataflow/ChaosFleshImportGEO.h b/Engine/Plugins/Experimental/ChaosFlesh/Source/ChaosFleshDeprecatedNodes/Public/Dataflow/ChaosFleshImportGEO.h
--- a/Engine/Plugins/Experimental/ChaosFlesh/Source/ChaosFleshDeprecatedNodes/Public/Dataflow/ChaosFleshImportGEO.h
+++ b/Engine/Plugins/Experimental/ChaosFlesh/Source/ChaosFleshDeprecatedNodes/Public/Dataflow/ChaosFleshImportGEO.h
@@ -12,7 +12,7 @@ DECLARE_LOG_CATEGORY_EXTERN(LogImportGEO, Verbose, All);
 
 
 /** Transient Dataflow container for \c TMap<FString, int32>. */
-USTRUCT(meta = (DataflowFlesh))
+USTRUCT(meta = (DataflowFlesh, Deprecated = "5.4"))
 struct FGEOMapStringInt
 {
 	GENERATED_USTRUCT_BODY()
@@ -20,7 +20,7 @@ struct FGEOMapStringInt
 };
 
 /** Transient Dataflow container for \c TMap<FString, TArray<int32>>. */
-USTRUCT(meta = (DataflowFlesh))
+USTRUCT(meta = (DataflowFlesh, Deprecated = "5.4"))
 struct FGEOMapStringArrayInt
 {
 	GENERATED_USTRUCT_BODY()
@@ -28,7 +28,7 @@ struct FGEOMapStringArrayInt
 };
 
 /** Transient Dataflow container for \c TMap<FString, TArray<float>>. */
-USTRUCT(meta = (DataflowFlesh))
+USTRUCT(meta = (DataflowFlesh, Deprecated = "5.4"))
 struct FGEOMapStringArrayFloat
 {
 	GENERATED_USTRUCT_BODY()
@@ -36,7 +36,7 @@ struct FGEOMapStringArrayFloat
 };
 
 /** Extract a named integer from the results of an ImportGEO node. */
-USTRUCT(meta = (DataflowFlesh))
+USTRUCT(meta = (DataflowFlesh, Deprecated = "5.4"))
 struct FExtractGEOInt : public FDataflowNode
 {
 	GENERATED_USTRUCT_BODY()
@@ -66,7 +66,7 @@ public:
 };
 
 /** Extract a named integer array from the results of an ImportGEO node. */
-USTRUCT(meta = (DataflowFlesh))
+USTRUCT(meta = (DataflowFlesh, Deprecated = "5.4"))
 struct FExtractGEOIntVector : public FDataflowNode
 {
 	GENERATED_USTRUCT_BODY()
@@ -96,7 +96,7 @@ public:
 };
 
 /** Extract a named float array from the results of an ImportGEO node. */
-USTRUCT(meta = (DataflowFlesh))
+USTRUCT(meta = (DataflowFlesh, Deprecated = "5.4"))
 struct FExtractGEOFloatVector : public FDataflowNode
 {
 	GENERATED_USTRUCT_BODY()
@@ -126,7 +126,7 @@ public:
 };
 
 /** Import data from GEO file. */
-USTRUCT(meta = (DataflowFlesh))
+USTRUCT(meta = (DataflowFlesh, Deprecated = "5.4"))
 struct FImportGEO : public FDataflowNode
 {
 	GENERATED_USTRUCT_BODY()
diff --git a/Engine/Plugins/Experimental/ChaosFlesh/Source/ChaosFleshDeprecatedNodes/Public/Dataflow/ChaosFleshKinematicOriginInsertionInitializationNode.h b/Engine/Plugins/Experimental/ChaosFlesh/Source/ChaosFleshDeprecatedNodes/Public/Dataflow/ChaosFleshKinematicOriginInsertionInitializationNode.h
--- a/Engine/Plugins/Experimental/ChaosFlesh/Source/ChaosFleshDeprecatedNodes/Public/Dataflow/ChaosFleshKinematicOriginInsertionInitializationNode.h
+++ b/Engine/Plugins/Experimental/ChaosFlesh/Source/ChaosFleshDeprecatedNodes/Public/Dataflow/ChaosFleshKinematicOriginInsertionInitializationNode.h
@@ -10,7 +10,7 @@
 class USkeletalMesh;
 
 // @todo(deprecate), rename to FKinematicSomethingBetterThanThisNameDataflowNode
-USTRUCT(meta = (DataflowFlesh))
+USTRUCT(meta = (DataflowFlesh, Deprecated = "5.4"))
 struct FKinematicOriginInsertionInitializationDataflowNode : public FDataflowNode
 {
 	GENERATED_USTRUCT_BODY()
diff --git a/Engine/Plugins/Experimental/ChaosFlesh/Source/ChaosFleshDeprecatedNodes/Public/Dataflow/ChaosFleshKinematicTetrahedralConstraintNode.h b/Engine/Plugins/Experimental/ChaosFlesh/Source/ChaosFleshDeprecatedNodes/Public/Dataflow/ChaosFleshKinematicTetrahedralConstraintNode.h
--- a/Engine/Plugins/Experimental/ChaosFlesh/Source/ChaosFleshDeprecatedNodes/Public/Dataflow/ChaosFleshKinematicTetrahedralConstraintNode.h
+++ b/Engine/Plugins/Experimental/ChaosFlesh/Source/ChaosFleshDeprecatedNodes/Public/Dataflow/ChaosFleshKinematicTetrahedralConstraintNode.h
@@ -11,7 +11,7 @@ class USkeletalMesh;
 
 // @todo(deprecate), rename to FKinematicTetrahedralConstraintDataflowNode
 // ... This should really be a mode on the KinematicConstraint. 
-USTRUCT(meta = (DataflowFlesh))
+USTRUCT(meta = (DataflowFlesh, Deprecated = "5.4"))
 struct FKinematicTetrahedralBindingsDataflowNode : public FDataflowNode
 {
 	GENERATED_USTRUCT_BODY()
