ChaosFlesh: Add vertex selection set into position target node

diff --git a/Engine/Plugins/Experimental/ChaosFlesh/Source/ChaosFleshNodes/Private/Dataflow/ChaosFleshSetVertexTrianglePositionTargetBindingNode.cpp b/Engine/Plugins/Experimental/ChaosFlesh/Source/ChaosFleshNodes/Private/Dataflow/ChaosFleshSetVertexTrianglePositionTargetBindingNode.cpp
--- a/Engine/Plugins/Experimental/ChaosFlesh/Source/ChaosFleshNodes/Private/Dataflow/ChaosFleshSetVertexTrianglePositionTargetBindingNode.cpp
+++ b/Engine/Plugins/Experimental/ChaosFlesh/Source/ChaosFleshNodes/Private/Dataflow/ChaosFleshSetVertexTrianglePositionTargetBindingNode.cpp
@@ -6,23 +6,30 @@
 #include "GeometryCollection/Facades/CollectionMeshFacade.h"
 #include "GeometryCollection/Facades/CollectionPositionTargetFacade.h"
 #include "ChaosFlesh/TetrahedralCollection.h"
+#include "ChaosFlesh/ChaosFleshCollectionFacade.h"
 
 #include UE_INLINE_GENERATED_CPP_BY_NAME(ChaosFleshSetVertexTrianglePositionTargetBindingNode)
 
-//DEFINE_LOG_CATEGORY_STATIC(ChaosFleshSetVertexTrianglePositionTargetBindingNodeLog, Log, All);
-
 void FSetVertexTrianglePositionTargetBindingDataflowNode::Evaluate(Dataflow::FContext& Context, const FDataflowOutput* Out) const
 {
 	if (Out->IsA<DataType>(&Collection))
 	{
 		DataType InCollection = GetValue<DataType>(Context, &Collection);
 
+		TUniquePtr<FFleshCollection> InFleshCollection(GetValue<DataType>(Context, &Collection).NewCopy<FFleshCollection>());
+
+		Chaos::FFleshCollectionFacade TetCollection(*InFleshCollection);
 		if (TManagedArray<FVector3f>* Vertices = InCollection.FindAttribute<FVector3f>("Vertex", FGeometryCollection::VerticesGroup))
+		{
+			if (TManagedArray<FIntVector>* Indices = InCollection.FindAttribute<FIntVector>("Indices", FGeometryCollection::FacesGroup))
 			{
-				if (TManagedArray<FIntVector>* Indices = InCollection.FindAttribute<FIntVector>("Indices", FGeometryCollection::FacesGroup))
+				if (TetCollection.IsTetrahedronValid())
 				{
+					TArray<FVector3f> Vertex = TetCollection.Vertex.Get().GetConstArray();
+					TetCollection.ComponentSpaceVertices(Vertex);
 					GeometryCollection::Facades::FCollectionMeshFacade MeshFacade(InCollection);
 					TArray<int32> ComponentIndex = MeshFacade.GetGeometryGroupIndexArray();
+					TArray<Chaos::TVector<float, 3>> IndicesPositions; 
 					TArray<Chaos::TVector<int32, 3>> IndicesArray;
 					for (int32 i = 0; i < Indices->Num(); i++)
 					{
@@ -31,9 +38,9 @@ void FSetVertexTrianglePositionTargetBindingDataflowNode::Evaluate(Dataflow::FCo
 						{
 							CurrentIndices[j] = (*Indices)[i][j];
 						}
-						if (CurrentIndices[0] != -1
-							&& CurrentIndices[1] != -1
-							&& CurrentIndices[2] != -1)
+						if (CurrentIndices[0] != INDEX_NONE
+							&& CurrentIndices[1] != INDEX_NONE
+							&& CurrentIndices[2] != INDEX_NONE)
 						{
 							IndicesArray.Emplace(CurrentIndices);
 						}
@@ -49,7 +56,7 @@ void FSetVertexTrianglePositionTargetBindingDataflowNode::Evaluate(Dataflow::FCo
 							ActualParticleCount += 1;
 						}
 					}
-					TArray<Chaos::TVector<float, 3>> IndicesPositions; 
+					
 					IndicesPositions.SetNum(ActualParticleCount);
 					TArray<int32> IndicesMap;
 					IndicesMap.SetNum(ActualParticleCount);
@@ -63,40 +70,83 @@ void FSetVertexTrianglePositionTargetBindingDataflowNode::Evaluate(Dataflow::FCo
 							CurrentParticleIndex += 1;
 						}
 					}
+					TArray<Chaos::TSphere<Chaos::FReal, 3>*> VertexSpherePtrs;
+					TArray<Chaos::TSphere<Chaos::FReal, 3>> VertexSpheres;
 					Chaos::FReal SphereRadius = (Chaos::FReal)0.;
-
 					Chaos::TVec3<float> CoordMaxs(-FLT_MAX);
 					Chaos::TVec3<float> CoordMins(FLT_MAX);
-					for (int32 i = 0; i < IndicesPositions.Num(); i++)
+					if (FindInput(&VertexSelectionSetIn) && FindInput(&VertexSelectionSetIn)->GetConnection())
 					{
-						for (int32 j = 0; j < 3; j++)
+						TArray<int32> VertexSelectionSet = GetValue<TArray<int32>>(Context, &VertexSelectionSetIn);
+						TArray<Chaos::TVector<float, 3>> VertexPositions;
+						VertexPositions.SetNum(VertexSelectionSet.Num());
+						for (int32 i = 0; i < VertexSelectionSet.Num(); i++)
 						{
-							if (IndicesPositions[i][j] > CoordMaxs[j])
+							if (VertexSelectionSet[i] > INDEX_NONE && VertexSelectionSet[i] < Vertex.Num())
 							{
-								CoordMaxs[j] = IndicesPositions[i][j];
+								VertexPositions[i] = Vertex[VertexSelectionSet[i]];
 							}
-							if (IndicesPositions[i][j] < CoordMins[j])
+						}
+						for (int32 i = 0; i < VertexPositions.Num(); i++)
+						{
+							for (int32 j = 0; j < 3; j++)
 							{
-								CoordMins[j] = IndicesPositions[i][j];
+								if (VertexPositions[i][j] > CoordMaxs[j])
+								{
+									CoordMaxs[j] = VertexPositions[i][j];
+								}
+								if (VertexPositions[i][j] < CoordMins[j])
+								{
+									CoordMins[j] = VertexPositions[i][j];
+								}
 							}
 						}
-					}
-					Chaos::TVec3<float> CoordDiff = (CoordMaxs - CoordMins) * VertexRadiusRatio;
-					SphereRadius = Chaos::FReal(FGenericPlatformMath::Min(CoordDiff[0], FGenericPlatformMath::Min(CoordDiff[1], CoordDiff[2])));
-
-					TArray<Chaos::TSphere<Chaos::FReal, 3>*> VertexSpherePtrs;
-					TArray<Chaos::TSphere<Chaos::FReal, 3>> VertexSpheres;
+						Chaos::TVec3<float> CoordDiff = (CoordMaxs - CoordMins) * VertexRadiusRatio;
+						SphereRadius = Chaos::FReal(FGenericPlatformMath::Min(CoordDiff[0], FGenericPlatformMath::Min(CoordDiff[1], CoordDiff[2])));
 
-					VertexSpheres.Init(Chaos::TSphere<Chaos::FReal, 3>(Chaos::TVec3<Chaos::FReal>(0), SphereRadius), IndicesPositions.Num());
-					VertexSpherePtrs.SetNum(IndicesPositions.Num());
+						VertexSpheres.Init(Chaos::TSphere<Chaos::FReal, 3>(Chaos::TVec3<Chaos::FReal>(0), SphereRadius), VertexPositions.Num());
+						VertexSpherePtrs.SetNum(VertexPositions.Num());
 
-					for (int32 i = 0; i < IndicesPositions.Num(); i++)
+						for (int32 i = 0; i < VertexPositions.Num(); i++)
+						{
+							Chaos::TVec3<Chaos::FReal> SphereCenter(VertexPositions[i]);
+							Chaos::TSphere<Chaos::FReal, 3> VertexSphere(SphereCenter, SphereRadius);
+							VertexSpheres[i] = Chaos::TSphere<Chaos::FReal, 3>(SphereCenter, SphereRadius);
+							VertexSpherePtrs[i] = &VertexSpheres[i];
+						}
+						IndicesMap = VertexSelectionSet;
+					}
+					else
 					{
-						Chaos::TVec3<Chaos::FReal> SphereCenter(IndicesPositions[i]);
-						Chaos::TSphere<Chaos::FReal, 3> VertexSphere(SphereCenter, SphereRadius);
-						VertexSpheres[i] = Chaos::TSphere<Chaos::FReal, 3>(SphereCenter, SphereRadius);
-						VertexSpherePtrs[i] = &VertexSpheres[i];
+						for (int32 i = 0; i < IndicesPositions.Num(); i++)
+						{
+							for (int32 j = 0; j < 3; j++)
+							{
+								if (IndicesPositions[i][j] > CoordMaxs[j])
+								{
+									CoordMaxs[j] = IndicesPositions[i][j];
+								}
+								if (IndicesPositions[i][j] < CoordMins[j])
+								{
+									CoordMins[j] = IndicesPositions[i][j];
+								}
+							}
+						}
+						Chaos::TVec3<float> CoordDiff = (CoordMaxs - CoordMins) * VertexRadiusRatio;
+						SphereRadius = Chaos::FReal(FGenericPlatformMath::Min(CoordDiff[0], FGenericPlatformMath::Min(CoordDiff[1], CoordDiff[2])));
+
+						VertexSpheres.Init(Chaos::TSphere<Chaos::FReal, 3>(Chaos::TVec3<Chaos::FReal>(0), SphereRadius), IndicesPositions.Num());
+						VertexSpherePtrs.SetNum(IndicesPositions.Num());
+
+						for (int32 i = 0; i < IndicesPositions.Num(); i++)
+						{
+							Chaos::TVec3<Chaos::FReal> SphereCenter(IndicesPositions[i]);
+							Chaos::TSphere<Chaos::FReal, 3> VertexSphere(SphereCenter, SphereRadius);
+							VertexSpheres[i] = Chaos::TSphere<Chaos::FReal, 3>(SphereCenter, SphereRadius);
+							VertexSpherePtrs[i] = &VertexSpheres[i];
+						}
 					}
+
 					Chaos::TBoundingVolumeHierarchy<
 						TArray<Chaos::TSphere<Chaos::FReal, 3>*>,
 						TArray<int32>,
@@ -108,9 +158,9 @@ void FSetVertexTrianglePositionTargetBindingDataflowNode::Evaluate(Dataflow::FCo
 
 					for (int32 i = 0; i < Indices->Num(); i++)
 					{
-						TArray<int32> TriangleIntersections0 = VertexBVH.FindAllIntersections((*Vertices)[(*Indices)[i][0]]);
-						TArray<int32> TriangleIntersections1 = VertexBVH.FindAllIntersections((*Vertices)[(*Indices)[i][1]]);
-						TArray<int32> TriangleIntersections2 = VertexBVH.FindAllIntersections((*Vertices)[(*Indices)[i][2]]);
+						TArray<int32> TriangleIntersections0 = VertexBVH.FindAllIntersections(Vertex[(*Indices)[i][0]]);
+						TArray<int32> TriangleIntersections1 = VertexBVH.FindAllIntersections(Vertex[(*Indices)[i][1]]);
+						TArray<int32> TriangleIntersections2 = VertexBVH.FindAllIntersections(Vertex[(*Indices)[i][2]]);
 						TriangleIntersections0.Sort();
 						TriangleIntersections1.Sort();
 						TriangleIntersections2.Sort();
@@ -131,11 +181,11 @@ void FSetVertexTrianglePositionTargetBindingDataflowNode::Evaluate(Dataflow::FCo
 						Chaos::TVector<float, 3> ClosestBary(0.f);
 						for (int32 j = 0; j < TriangleIntersections.Num(); j++)
 						{
-							if (ComponentIndex[IndicesMap[TriangleIntersections[j]]] >= 0 && TriangleIndex >= 0 && ComponentIndex[IndicesMap[TriangleIntersections[j]]] != TriangleIndex)
+							if (ComponentIndex[IndicesMap[TriangleIntersections[j]]] > INDEX_NONE && TriangleIndex > INDEX_NONE && ComponentIndex[IndicesMap[TriangleIntersections[j]]] != TriangleIndex)
 							{
-								Chaos::TVector<float, 3> Bary, TriPos0((*Vertices)[(*Indices)[i][0]]), TriPos1((*Vertices)[(*Indices)[i][1]]), TriPos2((*Vertices)[(*Indices)[i][2]]), ParticlePos((*Vertices)[IndicesMap[TriangleIntersections[j]]]);
+								Chaos::TVector<float, 3> Bary, TriPos0(Vertex[(*Indices)[i][0]]), TriPos1(Vertex[(*Indices)[i][1]]), TriPos2(Vertex[(*Indices)[i][2]]), ParticlePos(Vertex[IndicesMap[TriangleIntersections[j]]]);
 								Chaos::TVector<Chaos::FRealSingle, 3> ClosestPoint = Chaos::FindClosestPointAndBaryOnTriangle(TriPos0, TriPos1, TriPos2, ParticlePos, Bary);
-								Chaos::FRealSingle CurrentDistance = ((*Vertices)[IndicesMap[TriangleIntersections[j]]] - ClosestPoint).Size();
+								Chaos::FRealSingle CurrentDistance = (Vertex[IndicesMap[TriangleIntersections[j]]] - ClosestPoint).Size();
 								if (CurrentDistance < MinDis)
 								{
 									MinDis = CurrentDistance;
@@ -179,6 +229,7 @@ void FSetVertexTrianglePositionTargetBindingDataflowNode::Evaluate(Dataflow::FCo
 						}
 					}
 				}
+			}
 		}
 		SetValue(Context, MoveTemp(InCollection), &Collection);
 	}
diff --git a/Engine/Plugins/Experimental/ChaosFlesh/Source/ChaosFleshNodes/Public/Dataflow/ChaosFleshSetVertexTrianglePositionTargetBindingNode.h b/Engine/Plugins/Experimental/ChaosFlesh/Source/ChaosFleshNodes/Public/Dataflow/ChaosFleshSetVertexTrianglePositionTargetBindingNode.h
--- a/Engine/Plugins/Experimental/ChaosFlesh/Source/ChaosFleshNodes/Public/Dataflow/ChaosFleshSetVertexTrianglePositionTargetBindingNode.h
+++ b/Engine/Plugins/Experimental/ChaosFlesh/Source/ChaosFleshNodes/Public/Dataflow/ChaosFleshSetVertexTrianglePositionTargetBindingNode.h
@@ -25,8 +25,8 @@ public:
 	UPROPERTY(EditAnywhere, Category = "Dataflow")
 		float PositionTargetStiffness = 10000.f;
 
-	//UPROPERTY(EditAnywhere, Category = "Dataflow", meta = (ClampMin = "0.0", ClampMax = "1.0", UIMin = "0.0", UIMax = "1.0"))
-	//	float TriangleRadiusPadding = .1f;
+	UPROPERTY(meta = (DataflowInput, DisplayName = "VertexSelectionSet"))
+		TArray<int32> VertexSelectionSetIn;
 
 	UPROPERTY(EditAnywhere, Category = "Dataflow", meta = (ClampMin = "0.0", ClampMax = "2.0", UIMin = "0.0", UIMax = "2.0"))
 		float VertexRadiusRatio = .001f;
@@ -36,6 +36,7 @@ public:
 		: FDataflowNode(InParam, InGuid)
 	{
 		RegisterInputConnection(&Collection);
+		RegisterInputConnection(&VertexSelectionSetIn);
 		RegisterOutputConnection(&Collection, &Collection);
 	}
 
diff --git a/Engine/Plugins/Experimental/GeometryCollectionPlugin/Source/GeometryCollectionNodes/Private/Dataflow/GeometryCollectionNodesPlugin.cpp b/Engine/Plugins/Experimental/GeometryCollectionPlugin/Source/GeometryCollectionNodes/Private/Dataflow/GeometryCollectionNodesPlugin.cpp
--- a/Engine/Plugins/Experimental/GeometryCollectionPlugin/Source/GeometryCollectionNodes/Private/Dataflow/GeometryCollectionNodesPlugin.cpp
+++ b/Engine/Plugins/Experimental/GeometryCollectionPlugin/Source/GeometryCollectionNodes/Private/Dataflow/GeometryCollectionNodesPlugin.cpp
@@ -27,6 +27,7 @@
 #include "Dataflow/GeometryCollectionDebugNodes.h"
 #include "Dataflow/GeometryCollectionVertexScalarToVertexIndicesNode.h"
 #include "Dataflow/GeometryCollectionTransferVertexScalarAttributeNode.h"
+#include "Dataflow/GeometryCollectionTriangleBoundaryIndicesNode.h"
 
 #define LOCTEXT_NAMESPACE "DataflowNodes"
 
@@ -55,6 +56,7 @@ void IGeometryCollectionNodesPlugin::StartupModule()
 	DATAFLOW_NODE_REGISTER_CREATION_FACTORY(FSkeletonToCollectionDataflowNode);
 	DATAFLOW_NODE_REGISTER_CREATION_FACTORY(FSkeletalMeshToCollectionDataflowNode);
 	DATAFLOW_NODE_REGISTER_CREATION_FACTORY(FGeometryCollectionTransferVertexScalarAttributeNode);
+	DATAFLOW_NODE_REGISTER_CREATION_FACTORY(FTriangleBoundaryIndicesNode);
 }
 
 void IGeometryCollectionNodesPlugin::ShutdownModule()
diff --git a/Engine/Plugins/Experimental/GeometryCollectionPlugin/Source/GeometryCollectionNodes/Private/Dataflow/GeometryCollectionTriangleBoundaryIndicesNode.cpp b/Engine/Plugins/Experimental/GeometryCollectionPlugin/Source/GeometryCollectionNodes/Private/Dataflow/GeometryCollectionTriangleBoundaryIndicesNode.cpp
new file mode 100644
--- /dev/null
+++ b/Engine/Plugins/Experimental/GeometryCollectionPlugin/Source/GeometryCollectionNodes/Private/Dataflow/GeometryCollectionTriangleBoundaryIndicesNode.cpp
@@ -0,0 +1,49 @@
+// Copyright Epic Games, Inc. All Rights Reserved.
+
+#include "Dataflow/GeometryCollectionTriangleBoundaryIndicesNode.h"
+
+#include "Chaos/BoundingVolumeHierarchy.h"
+#include "GeometryCollection/Facades/CollectionMeshFacade.h"
+#include "GeometryCollection/Facades/CollectionPositionTargetFacade.h"
+#include "Chaos/Utilities.h"
+
+#include UE_INLINE_GENERATED_CPP_BY_NAME(GeometryCollectionTriangleBoundaryIndicesNode)
+
+void FTriangleBoundaryIndicesNode::Evaluate(Dataflow::FContext& Context, const FDataflowOutput* Out) const
+{
+	if (Out->IsA<TArray<int32>>(&BoundaryIndicesOut))
+	{
+		const FManagedArrayCollection& InCollection = GetValue<FManagedArrayCollection>(Context, &Collection);
+
+		TArray<Chaos::TVec3<int32>> TriangleMeshArray;
+
+		if (const TManagedArray<int32>* TriangleMeshIndices = InCollection.FindAttribute<int32>("ObjectIndices", "TriangleMesh"))
+		{
+			if (const TManagedArray<FIntVector>* Indices = InCollection.FindAttribute<FIntVector>("Indices", FGeometryCollection::FacesGroup))
+			{
+				if (const TManagedArray<int32>* FaceStarts = InCollection.FindAttribute<int32>("FaceStart", FGeometryCollection::GeometryGroup))
+				{
+					if (const TManagedArray<int32>* FaceCounts = InCollection.FindAttribute<int32>("FaceCount", FGeometryCollection::GeometryGroup))
+					{
+						for (int32 i = 0; i < TriangleMeshIndices->Num(); i++)
+						{
+							const int32 ObjectIndex = (*TriangleMeshIndices)[i];
+							
+							const int32 FaceStartIndex = (*FaceStarts)[ObjectIndex];
+							const int32 FaceNum = (*FaceCounts)[ObjectIndex];
+							for (int32 e = FaceStartIndex; e < FaceStartIndex + FaceNum; e++)
+							{
+								Chaos::TVec3<int32> IndicesTemp((*Indices)[e][0], (*Indices)[e][1], (*Indices)[e][2]);
+								TriangleMeshArray.Emplace(IndicesTemp);
+							}
+						}	
+					}
+				}
+			}
+		}
+
+		TArray<int32> IndicesOut = Chaos::Utilities::ComputeBoundaryNodes(TriangleMeshArray);
+
+		SetValue<TArray<int32>>(Context, MoveTemp(IndicesOut), &BoundaryIndicesOut);
+	}
+}
diff --git a/Engine/Plugins/Experimental/GeometryCollectionPlugin/Source/GeometryCollectionNodes/Public/Dataflow/GeometryCollectionTriangleBoundaryIndicesNode.h b/Engine/Plugins/Experimental/GeometryCollectionPlugin/Source/GeometryCollectionNodes/Public/Dataflow/GeometryCollectionTriangleBoundaryIndicesNode.h
new file mode 100644
--- /dev/null
+++ b/Engine/Plugins/Experimental/GeometryCollectionPlugin/Source/GeometryCollectionNodes/Public/Dataflow/GeometryCollectionTriangleBoundaryIndicesNode.h
@@ -0,0 +1,39 @@
+// Copyright Epic Games, Inc. All Rights Reserved.
+
+#pragma once 
+
+#include "CoreMinimal.h"
+#include "Dataflow/DataflowCore.h"
+#include "Dataflow/DataflowEngine.h"
+
+#include "GeometryCollectionTriangleBoundaryIndicesNode.generated.h"
+
+//Outputs boundary nodes of a triangle mesh
+USTRUCT(meta = (DataflowFlesh))
+struct FTriangleBoundaryIndicesNode : public FDataflowNode
+{
+	GENERATED_USTRUCT_BODY()
+	DATAFLOW_NODE_DEFINE_INTERNAL(FTriangleBoundaryIndicesNode, "TriangleBoundaryIndices", "Geometry Collection|Utiltites", "")
+		DATAFLOW_NODE_RENDER_TYPE(FGeometryCollection::StaticType(), "Collection")
+		
+
+public:
+	typedef FManagedArrayCollection DataType;
+
+
+	UPROPERTY(meta = (DataflowInput, DisplayName = "Collection", DataflowPassthrough = "Collection"))
+		FManagedArrayCollection Collection;
+
+	UPROPERTY(meta = (DataflowOutput, DisplayName = "BoundaryIndices"))
+		TArray<int32> BoundaryIndicesOut;
+
+
+	FTriangleBoundaryIndicesNode(const Dataflow::FNodeParameters& InParam, FGuid InGuid = FGuid::NewGuid())
+		: FDataflowNode(InParam, InGuid)
+	{
+		RegisterInputConnection(&Collection);
+		RegisterOutputConnection(&BoundaryIndicesOut);
+	}
+
+	virtual void Evaluate(Dataflow::FContext& Context, const FDataflowOutput* Out) const override;
+};
\ No newline at end of file
diff --git a/Engine/Source/Runtime/Experimental/Chaos/Public/Chaos/Utilities.h b/Engine/Source/Runtime/Experimental/Chaos/Public/Chaos/Utilities.h
--- a/Engine/Source/Runtime/Experimental/Chaos/Public/Chaos/Utilities.h
+++ b/Engine/Source/Runtime/Experimental/Chaos/Public/Chaos/Utilities.h
@@ -1124,6 +1124,52 @@ namespace Chaos
 			ConnectedComponentsDFSIterative(AdjacencyList, ConnectedComponents);
 		}
 
+		inline TArray<int32> ComputeBoundaryNodes(const TArray<TVec3<int32>>& TriangleMesh)
+		{
+			TMap<TPair<int32, int32>, int32> FaceCountPerEdge;
+		
+			auto AddTriToFaceCountPerEdge = [&FaceCountPerEdge](const TVec3<int32> &Face)
+			{
+				auto AddEdge = [&FaceCountPerEdge](int32 A, int32 B)
+				{
+					FaceCountPerEdge.FindOrAdd(TPair<int32, int32>(FMath::Min(A, B), FMath::Max(A, B)))++;
+				};
+				AddEdge(Face[0], Face[1]);
+				AddEdge(Face[0], Face[2]);
+				AddEdge(Face[1], Face[2]);
+			};
+
+			auto CandidateTriWouldMakeNonManifoldEdge = [&FaceCountPerEdge](int32 A, int32 B, int32 C)
+			{
+				auto GetCount = [&FaceCountPerEdge](int32 InnerA, int32 InnerB)
+				{
+					int32 *Count = FaceCountPerEdge.Find(TPair<int32, int32>(FMath::Min(InnerA, InnerB), FMath::Max(InnerA, InnerB)));
+					return Count ? *Count : 0;
+				};
+				return GetCount(A, B) > 1 || GetCount(B, C) > 1 || GetCount(C, A) > 1;
+			};
+
+			for (const TVec3<int32> &Face : TriangleMesh)
+			{
+				AddTriToFaceCountPerEdge(Face);
+			}
+
+			TArray<TPair<int32, int32>> AllEdges;
+			TSet<int32> BoundaryVerts;
+			int32 NumKeys = FaceCountPerEdge.GetKeys(AllEdges);
+
+			for (const TPair<int32, int32>& Edge: AllEdges)
+			{
+				if (FaceCountPerEdge[Edge] == 1)
+				{
+					BoundaryVerts.Emplace(Edge.Key);
+					BoundaryVerts.Emplace(Edge.Value);
+				}
+			}
+
+			return BoundaryVerts.Array();
+		}
+
 		inline TArray<TArray<int32>> ComputeIncidentElements(const TArray<TArray<int32>>& Constraints, TArray<TArray<int32>>* LocalIndex = nullptr)
 		{
 			int32 MaxIdx = 0;
