ChaosFlesh: Visualize muscle attachments and fiber field

diff --git a/Engine/Plugins/Experimental/ChaosFlesh/Source/ChaosFleshNodes/Private/Dataflow/ChaosFleshComputeFiberFieldNode.cpp b/Engine/Plugins/Experimental/ChaosFlesh/Source/ChaosFleshNodes/Private/Dataflow/ChaosFleshComputeFiberFieldNode.cpp
--- a/Engine/Plugins/Experimental/ChaosFlesh/Source/ChaosFleshNodes/Private/Dataflow/ChaosFleshComputeFiberFieldNode.cpp
+++ b/Engine/Plugins/Experimental/ChaosFlesh/Source/ChaosFleshNodes/Private/Dataflow/ChaosFleshComputeFiberFieldNode.cpp
@@ -127,10 +127,11 @@ void FComputeFiberFieldNode::Evaluate(Dataflow::FContext& Context, const FDatafl
 		// Do the thing
 		//
 
-		TArray<FVector3f> FiberDirs =
-			ComputeFiberField(*Elements, *Vertex, *IncidentElements, *IncidentElementsLocalIndex, 
+		TArray<FVector3f> FiberDirs;
+		TArray<float> MuscleAttachmentScalarFieldTArray; //continuous field where origin = 1, insertion = 2, othernodes = 0
+		ComputeFiberField(*Elements, *Vertex, *IncidentElements, *IncidentElementsLocalIndex, 
 				Origin ? Origin->GetConstArray() : InOriginIndices,
-				Insertion ? Insertion->GetConstArray() : InInsertionIndices);
+				Insertion ? Insertion->GetConstArray() : InInsertionIndices, FiberDirs, MuscleAttachmentScalarFieldTArray);
 
 		//
 		// Set output(s)
@@ -145,6 +146,21 @@ void FComputeFiberFieldNode::Evaluate(Dataflow::FContext& Context, const FDatafl
 		}
 		(*FiberDirections) = MoveTemp(FiberDirs);
 
+		TManagedArray<FLinearColor>* Color =
+			InCollection.FindAttribute<FLinearColor>("Color", FGeometryCollection::VerticesGroup);
+		if (!Color)
+		{
+			Color =
+				&InCollection.AddAttribute<FLinearColor>("Color", FGeometryCollection::VerticesGroup);
+		}
+		for (int32 i = 0; i < Color->Num(); ++i)
+		{
+			float s = MuscleAttachmentScalarFieldTArray[i];
+			if (s > 0) // 1 <= s <= 2 if muscle
+			{
+				(*Color)[i] = FLinearColor(FVector(s-1, 0, 2 - s));
+			}
+		}
 		Out->SetValue(MoveTemp(InCollection), Context);
 	}
 }
@@ -164,16 +180,16 @@ FComputeFiberFieldNode::GetNonZeroIndices(const TArray<uint8>& Map) const
 	return Indices;
 }
 
-TArray<FVector3f>
-FComputeFiberFieldNode::ComputeFiberField(
+void FComputeFiberFieldNode::ComputeFiberField(
 	const TManagedArray<FIntVector4>& Elements,
 	const TManagedArray<FVector3f>& Vertex,
 	const TManagedArray<TArray<int32>>& IncidentElements,
 	const TManagedArray<TArray<int32>>& IncidentElementsLocalIndex,
 	const TArray<int32>& Origin,
-	const TArray<int32>& Insertion) const
+	const TArray<int32>& Insertion,
+	TArray<FVector3f>& Directions,
+	TArray<float>& ScalarField) const
 {
-	TArray<FVector3f> Directions;
 	Chaos::ComputeFiberField<float>(
 		Elements.GetConstArray(),
 		Vertex.GetConstArray(),
@@ -182,7 +198,7 @@ FComputeFiberFieldNode::ComputeFiberField(
 		Origin,
 		Insertion,
 		Directions,
+		ScalarField,
 		MaxIterations,
 		Tolerance);
-	return Directions;
 }
diff --git a/Engine/Plugins/Experimental/ChaosFlesh/Source/ChaosFleshNodes/Private/Dataflow/ChaosFleshNodesPlugin.cpp b/Engine/Plugins/Experimental/ChaosFlesh/Source/ChaosFleshNodes/Private/Dataflow/ChaosFleshNodesPlugin.cpp
--- a/Engine/Plugins/Experimental/ChaosFlesh/Source/ChaosFleshNodes/Private/Dataflow/ChaosFleshNodesPlugin.cpp
+++ b/Engine/Plugins/Experimental/ChaosFlesh/Source/ChaosFleshNodes/Private/Dataflow/ChaosFleshNodesPlugin.cpp
@@ -71,6 +71,7 @@ void IChaosFleshNodesPlugin::StartupModule()
 	DATAFLOW_NODE_REGISTER_CREATION_FACTORY(FSkinSimulationPropertiesDataflowNodes);
 	DATAFLOW_NODE_REGISTER_CREATION_FACTORY(FTriangleMeshSimulationPropertiesDataflowNodes);
 	DATAFLOW_NODE_REGISTER_CREATION_FACTORY(FVisualizeFiberFieldNode);
+	DATAFLOW_NODE_REGISTER_CREATION_FACTORY(FVisualizePositionTargetsNode);
 }
 
 void IChaosFleshNodesPlugin::ShutdownModule()
diff --git a/Engine/Plugins/Experimental/ChaosFlesh/Source/ChaosFleshNodes/Private/Dataflow/ChaosFleshVisualizeFiberFieldNode.cpp b/Engine/Plugins/Experimental/ChaosFlesh/Source/ChaosFleshNodes/Private/Dataflow/ChaosFleshVisualizeFiberFieldNode.cpp
--- a/Engine/Plugins/Experimental/ChaosFlesh/Source/ChaosFleshNodes/Private/Dataflow/ChaosFleshVisualizeFiberFieldNode.cpp
+++ b/Engine/Plugins/Experimental/ChaosFlesh/Source/ChaosFleshNodes/Private/Dataflow/ChaosFleshVisualizeFiberFieldNode.cpp
@@ -1,6 +1,7 @@
 // Copyright Epic Games, Inc. All Rights Reserved.
 
 #include "Dataflow/ChaosFleshVisualizeFiberFieldNode.h"
+#include "GeometryCollection/Facades/CollectionPositionTargetFacade.h"
 
 #include "ChaosFlesh/TetrahedralCollection.h"
 
@@ -22,14 +23,17 @@ void FVisualizeFiberFieldNode::Evaluate(Dataflow::FContext& Context, const FData
 					ensureMsgf(Elements->Num() == FiberDirections->Num(), TEXT("Fiber direction has different size than elements"));
 					for (int32 ElemIndex = 0; ElemIndex < Elements->Num(); ElemIndex++)
 					{
-						FVector3f VectorStart = { 0,0,0 };
-						for (int32 LocalIndex = 0; LocalIndex < 4; LocalIndex++)
+						if ((*FiberDirections)[ElemIndex].Length() > 1e-10)
 						{
-							VectorStart += (*Vertex)[(*Elements)[ElemIndex][LocalIndex]];
+							FVector3f VectorStart = { 0,0,0 };
+							for (int32 LocalIndex = 0; LocalIndex < 4; LocalIndex++)
+							{
+								VectorStart += (*Vertex)[(*Elements)[ElemIndex][LocalIndex]];
+							}
+							VectorStart /= float(4);
+							FVector3f VectorEnd = VectorStart + (*FiberDirections)[ElemIndex] * VectorScale;
+							OutVectorField.AddVectorToField(VectorStart, VectorEnd);
 						}
-						VectorStart /= float(4);
-						FVector3f VectorEnd = VectorStart + (*FiberDirections)[ElemIndex] * VectorScale;
-						OutVectorField.AddVectorToField(VectorStart, VectorEnd);
 					}
 				}
 			}
@@ -38,3 +42,36 @@ void FVisualizeFiberFieldNode::Evaluate(Dataflow::FContext& Context, const FData
 		Out->SetValue(MoveTemp(OutVectorField), Context);
 	}
 }
+
+void FVisualizePositionTargetsNode::Evaluate(Dataflow::FContext& Context, const FDataflowOutput* Out) const
+{
+	if (Out->IsA<FFieldCollection>(&VectorField))
+	{
+		FManagedArrayCollection InCollection = GetValue<FManagedArrayCollection>(Context, &Collection);
+		FFieldCollection OutVectorField = VectorField;
+
+		if (TManagedArray<FVector3f>* Vertex = InCollection.FindAttribute<FVector3f>("Vertex", "Vertices"))
+		{
+			GeometryCollection::Facades::FPositionTargetFacade PositionTargets(InCollection);
+
+			// Read in position target info
+			for (int i = 0; i < PositionTargets.NumPositionTargets(); i++)
+			{
+				GeometryCollection::Facades::FPositionTargetsData DataPackage = PositionTargets.GetPositionTarget(i);			
+				FVector3f VectorStart = { 0,0,0 };
+				for (int32 LocalIndex = 0; LocalIndex < DataPackage.SourceIndex.Num(); LocalIndex++)
+				{
+					VectorStart += DataPackage.SourceWeights[LocalIndex] * (*Vertex)[DataPackage.SourceIndex[LocalIndex]];
+				}
+				FVector3f VectorEnd = { 0,0,0 };
+				for (int32 LocalIndex = 0; LocalIndex < DataPackage.TargetIndex.Num(); LocalIndex++)
+				{
+					VectorEnd += DataPackage.TargetWeights[LocalIndex] * (*Vertex)[DataPackage.TargetIndex[LocalIndex]];
+				}
+				OutVectorField.AddVectorToField(VectorStart, VectorEnd);
+			}
+		}
+
+		Out->SetValue(MoveTemp(OutVectorField), Context);
+	}
+}
\ No newline at end of file
diff --git a/Engine/Plugins/Experimental/ChaosFlesh/Source/ChaosFleshNodes/Public/Dataflow/ChaosFleshComputeFiberFieldNode.h b/Engine/Plugins/Experimental/ChaosFlesh/Source/ChaosFleshNodes/Public/Dataflow/ChaosFleshComputeFiberFieldNode.h
--- a/Engine/Plugins/Experimental/ChaosFlesh/Source/ChaosFleshNodes/Public/Dataflow/ChaosFleshComputeFiberFieldNode.h
+++ b/Engine/Plugins/Experimental/ChaosFlesh/Source/ChaosFleshNodes/Public/Dataflow/ChaosFleshComputeFiberFieldNode.h
@@ -62,11 +62,13 @@ public:
 
 	TArray<int32> GetNonZeroIndices(const TArray<uint8>& Map) const;
 
-	TArray<FVector3f> ComputeFiberField(
+	void ComputeFiberField(
 		const TManagedArray<FIntVector4>& Elements,
 		const TManagedArray<FVector3f>& Vertex,
 		const TManagedArray<TArray<int32>>& IncidentElements,
 		const TManagedArray<TArray<int32>>& IncidentElementsLocalIndex,
 		const TArray<int32>& Origin,
-		const TArray<int32>& Insertion) const;
+		const TArray<int32>& Insertion,
+		TArray<FVector3f>& Directions,
+		TArray<float>& ScalarField) const;
 };
diff --git a/Engine/Plugins/Experimental/ChaosFlesh/Source/ChaosFleshNodes/Public/Dataflow/ChaosFleshComputeMuscleActivationNode.h b/Engine/Plugins/Experimental/ChaosFlesh/Source/ChaosFleshNodes/Public/Dataflow/ChaosFleshComputeMuscleActivationNode.h
--- a/Engine/Plugins/Experimental/ChaosFlesh/Source/ChaosFleshNodes/Public/Dataflow/ChaosFleshComputeMuscleActivationNode.h
+++ b/Engine/Plugins/Experimental/ChaosFlesh/Source/ChaosFleshNodes/Public/Dataflow/ChaosFleshComputeMuscleActivationNode.h
@@ -22,10 +22,10 @@ public:
 	UPROPERTY(meta = (DataflowInput, DataflowOutput, DisplayName = "Collection", DataflowPassthrough = "Collection"))
 	FManagedArrayCollection Collection;
 
-	UPROPERTY(meta = (DataflowInput, DisplayName = "OriginIndicesIn"))
+	UPROPERTY(meta = (DataflowInput, DisplayName = "OriginVertexIndices"))
 		TArray<int32> OriginIndicesIn;
 
-	UPROPERTY(meta = (DataflowInput, DisplayName = "InsertionIndicesIn"))
+	UPROPERTY(meta = (DataflowInput, DisplayName = "InsertionVertexIndices"))
 		TArray<int32> InsertionIndicesIn;
 
 	FComputeMuscleActivationDataNode(const Dataflow::FNodeParameters& InParam, FGuid InGuid = FGuid::NewGuid())
diff --git a/Engine/Plugins/Experimental/ChaosFlesh/Source/ChaosFleshNodes/Public/Dataflow/ChaosFleshKinematicMuscleAttachmentsNode.h b/Engine/Plugins/Experimental/ChaosFlesh/Source/ChaosFleshNodes/Public/Dataflow/ChaosFleshKinematicMuscleAttachmentsNode.h
--- a/Engine/Plugins/Experimental/ChaosFlesh/Source/ChaosFleshNodes/Public/Dataflow/ChaosFleshKinematicMuscleAttachmentsNode.h
+++ b/Engine/Plugins/Experimental/ChaosFlesh/Source/ChaosFleshNodes/Public/Dataflow/ChaosFleshKinematicMuscleAttachmentsNode.h
@@ -22,10 +22,10 @@ public:
 	UPROPERTY(meta = (DataflowInput, DataflowOutput, DisplayName = "Collection", DataflowPassthrough = "Collection"))
 		FManagedArrayCollection Collection;
 
-	UPROPERTY(meta = (DataflowInput, DisplayName = "OriginSelectionSet"))
+	UPROPERTY(meta = (DataflowInput, DisplayName = "OriginVertexIndices"))
 		TArray<int32> OriginVertexIndicesIn;
 
-	UPROPERTY(meta = (DataflowInput, DisplayName = "InsertionSelectionSet"))
+	UPROPERTY(meta = (DataflowInput, DisplayName = "InsertionVertexIndices"))
 		TArray<int32> InsertionVertexIndicesIn;
 
 	UPROPERTY(meta = (DataflowInput, DisplayName = "SkeletalMesh"))
diff --git a/Engine/Plugins/Experimental/ChaosFlesh/Source/ChaosFleshNodes/Public/Dataflow/ChaosFleshVisualizeFiberFieldNode.h b/Engine/Plugins/Experimental/ChaosFlesh/Source/ChaosFleshNodes/Public/Dataflow/ChaosFleshVisualizeFiberFieldNode.h
--- a/Engine/Plugins/Experimental/ChaosFlesh/Source/ChaosFleshNodes/Public/Dataflow/ChaosFleshVisualizeFiberFieldNode.h
+++ b/Engine/Plugins/Experimental/ChaosFlesh/Source/ChaosFleshNodes/Public/Dataflow/ChaosFleshVisualizeFiberFieldNode.h
@@ -38,3 +38,30 @@ public:
 
 	virtual void Evaluate(Dataflow::FContext& Context, const FDataflowOutput* Out) const override;
 };
+
+/**
+* Visualizes position target vectors from GeometryCollection.
+*/
+USTRUCT(meta = (DataflowFlesh))
+struct FVisualizePositionTargetsNode : public FDataflowNode
+{
+	GENERATED_USTRUCT_BODY()
+	DATAFLOW_NODE_DEFINE_INTERNAL(FVisualizePositionTargetsNode, "VisualizePositionTargets", "Flesh", "")
+		DATAFLOW_NODE_RENDER_TYPE(FFieldCollection::StaticType(), "VectorField")
+
+public:
+	UPROPERTY(meta = (DataflowInput, DisplayName = "Collection"))
+	FManagedArrayCollection Collection;
+
+	UPROPERTY(meta = (DataflowOutput, DisplayName = "VectorField"))
+	FFieldCollection VectorField;
+
+	FVisualizePositionTargetsNode(const Dataflow::FNodeParameters& InParam, FGuid InGuid = FGuid::NewGuid())
+		: FDataflowNode(InParam, InGuid)
+	{
+		RegisterInputConnection(&Collection);
+		RegisterOutputConnection(&VectorField);
+	}
+
+	virtual void Evaluate(Dataflow::FContext& Context, const FDataflowOutput* Out) const override;
+};
\ No newline at end of file
diff --git a/Engine/Source/Runtime/Experimental/Chaos/Public/Chaos/Math/Poisson.h b/Engine/Source/Runtime/Experimental/Chaos/Public/Chaos/Math/Poisson.h
--- a/Engine/Source/Runtime/Experimental/Chaos/Public/Chaos/Math/Poisson.h
+++ b/Engine/Source/Runtime/Experimental/Chaos/Public/Chaos/Math/Poisson.h
@@ -465,6 +465,7 @@ ComputeFiberField(
 	const TArray<int32>& Origins,
 	const TArray<int32>& Insertions,
 	TArray<TV>& Directions,
+	TArray<T>& ScalarField,
 	const int32 MaxIt=100,
 	const T Tol=T(1e-7))
 {
@@ -473,12 +474,12 @@ ComputeFiberField(
 	{
 		if (zero_non_boundary)
 			Fill(minus_u_bc, T(0));
-		// origin points have bc = 0
+		// origin points have bc = scale
 		for (int32 i = 0; i < Origins.Num(); i++)
-			minus_u_bc[Origins[i]] = T(0);
-		// insertion points have bc = scale
+			minus_u_bc[Origins[i]] = scale;
+		// insertion points have bc = 2*scale
 		for (int32 i = 0; i < Insertions.Num(); i++)
-			minus_u_bc[Insertions[i]] = scale;
+			minus_u_bc[Insertions[i]] = 2*scale;
 	};
 
 	// Origin and insertion points are fixed
@@ -520,6 +521,7 @@ ComputeFiberField(
 	//compute directions from scalar gradient
 	Directions.SetNumUninitialized(Mesh.Num());
 	Fill(Directions, TV(0));
+	ScalarField = u;
 
 	TArray<TV> grad_Nie_hat; grad_Nie_hat.SetNumUninitialized(d + 1);
 	grad_Nie_hat[0] = { T(-1),T(-1),T(-1) };
@@ -530,8 +532,7 @@ ComputeFiberField(
 	TV pseudo_direction = Vertices[Insertions[0]] - Vertices[Origins[0]];
 	pseudo_direction.Normalize();
 
-	//PhysicsParallelFor(Mesh.Num(), [&](const int32 e)
-	for(int32 e=0; e < Mesh.Num(); e++)
+	PhysicsParallelFor(Mesh.Num(), [&](const int32 e)
 	{
 		T Deinv[d * d]{
 			De_inverse[d * d * e + 0],
@@ -560,12 +561,15 @@ ComputeFiberField(
 		{
 			Directions[e] = gradient * T(1) / Len;
 		}
-		else 
+		else if (u[Mesh[e][0]] > 0)
 		{
-			std::cout << "zero fiber gradient" << std::endl;
 			Directions[e] = pseudo_direction;
 		}
-	}//);
+		else
+		{
+			Directions[e] = gradient;
+		}
+	});
 }
 
 // 9 point laplacian with dirichlet boundaries
