2- Chaos USD caching multiple stages

diff --git a/Engine/Plugins/Experimental/ChaosCaching/Source/ChaosCaching/Private/Chaos/CacheCollection.cpp b/Engine/Plugins/Experimental/ChaosCaching/Source/ChaosCaching/Private/Chaos/CacheCollection.cpp
--- a/Engine/Plugins/Experimental/ChaosCaching/Source/ChaosCaching/Private/Chaos/CacheCollection.cpp
+++ b/Engine/Plugins/Experimental/ChaosCaching/Source/ChaosCaching/Private/Chaos/CacheCollection.cpp
@@ -122,3 +122,18 @@ float UChaosCacheCollection::GetMaxDuration() const
     }
 	return MaxDuration;
 }
+
+uint32 UChaosCacheCollection::GetMaxNumFrames() const
+{
+	float MaxDuration = 0.0;
+	uint32 MaxFramesCount = 0;
+	for(const UChaosCache* CacheInstance : Caches)
+	{
+		if(CacheInstance->GetDuration() >= MaxDuration)
+		{
+			MaxDuration = CacheInstance->GetDuration();
+			MaxFramesCount = CacheInstance->NumRecordedFrames;
+		}
+	}
+	return MaxFramesCount;
+}
diff --git a/Engine/Plugins/Experimental/ChaosCaching/Source/ChaosCaching/Private/Chaos/CacheManagerActor.cpp b/Engine/Plugins/Experimental/ChaosCaching/Source/ChaosCaching/Private/Chaos/CacheManagerActor.cpp
--- a/Engine/Plugins/Experimental/ChaosCaching/Source/ChaosCaching/Private/Chaos/CacheManagerActor.cpp
+++ b/Engine/Plugins/Experimental/ChaosCaching/Source/ChaosCaching/Private/Chaos/CacheManagerActor.cpp
@@ -728,6 +728,9 @@ void AChaosCacheManager::OnStartFrameChanged(Chaos::FReal InTime)
 			{
 				if (Observed.BestFitAdapter->ValidForPlayback(Comp, Observed.Cache))
 				{
+					// the init for load is quite lightweight and should be done before each update of time
+					// @todo maybe rename SetRestState to LoadCacheAtTime
+					Observed.BestFitAdapter->InitializeForLoad(Comp, Observed);
 					Observed.BestFitAdapter->SetRestState(Comp, Observed.Cache, GetTransform(), InTime);
 				}
 				Observed.Cache->EndPlayback(Token);
@@ -843,10 +846,27 @@ FObservedComponent& AChaosCacheManager::AddNewObservedComponent(UPrimitiveCompon
 	return NewEntry;
 }
 
-FObservedComponent& AChaosCacheManager::FindOrAddObservedComponent(UPrimitiveComponent* InComponent)
+FObservedComponent& AChaosCacheManager::FindOrAddObservedComponent(UPrimitiveComponent* InComponent, const FName& CacheName, const bool bTransferSimulation)
 {
-	FObservedComponent* Found = FindObservedComponent(InComponent);
-	return Found ? *Found : AddNewObservedComponent(InComponent);
+	FObservedComponent* FoundComponent = FindObservedComponent(InComponent);
+	FObservedComponent* ObservedComponent = FoundComponent ? FoundComponent : &AddNewObservedComponent(InComponent);
+
+	ObservedComponent->bIsSimulating = (InComponent && bTransferSimulation) ? InComponent->BodyInstance.bSimulatePhysics : false;
+
+	if(CacheName != TEXT(""))
+	{
+		ObservedComponent->CacheName = CacheName;
+	}
+
+	return *ObservedComponent;
+}
+
+void AChaosCacheManager::RemoveObservedComponent(UPrimitiveComponent* PrimitiveComponent)
+{
+	ObservedComponents.RemoveAll([this, ToTest = PrimitiveComponent](const FObservedComponent& Item)
+	{
+		return Item.GetComponent(this) == ToTest;
+	});
 }
 
 void AChaosCacheManager::ClearObservedComponents()
diff --git a/Engine/Plugins/Experimental/ChaosCaching/Source/ChaosCaching/Private/Chaos/ChaosCache.cpp b/Engine/Plugins/Experimental/ChaosCaching/Source/ChaosCaching/Private/Chaos/ChaosCache.cpp
--- a/Engine/Plugins/Experimental/ChaosCaching/Source/ChaosCaching/Private/Chaos/ChaosCache.cpp
+++ b/Engine/Plugins/Experimental/ChaosCaching/Source/ChaosCaching/Private/Chaos/ChaosCache.cpp
@@ -218,7 +218,7 @@ bool UChaosCache::FlushPendingFrames_MainPass(TQueue<FPendingFrameWrite, Mode>&
 		}
 		return false;
 	};
-
+	
 	while (InPendingWrites.Dequeue(NewData))
 	{
 		QUICK_SCOPE_CYCLE_COUNTER(QSTAT_CacheFlushSingleFramePass2);
@@ -279,6 +279,8 @@ bool UChaosCache::FlushPendingFrames_MainPass(TQueue<FPendingFrameWrite, Mode>&
 				}
 			}
 		}
+		MinTime = FMath::Min(MinTime,NewData.Time);
+		MaxTime = FMath::Max(MaxTime,NewData.Time);
 
 		const int32 ParticleCount = NewData.PendingParticleData.Num();
 
@@ -381,7 +383,14 @@ void UChaosCache::FlushPendingFrames()
 			}
 		}
 
-		RecordedDuration = Max - Min;
+		RecordedDuration =  Max - Min;
+	}
+	// In case we are not recording to channels/particles/curves datas (USD) we can
+	// use the pending writes time to figure out the recorded duration.
+	// Should probably be the unique way of computing it. 
+	if(ParticleTracks.IsEmpty() && ChannelsTracks.IsEmpty() && (MaxTime > MinTime))
+	{
+		RecordedDuration = MaxTime-MinTime;
 	}
 }
 
@@ -397,6 +406,8 @@ FCacheUserToken UChaosCache::BeginRecord(const UPrimitiveComponent* InComponent,
 			// And there's no playbacks, we can proceed
 			// Setup the cache to begin recording
 			RecordedDuration = 0.0f;
+			MinTime = TNumericLimits<float>::Max();
+			MaxTime = TNumericLimits<float>::Lowest();
 			NumRecordedFrames = 0;
 			ParticleTracks.Reset();
 			ChannelsTracks.Reset();
diff --git a/Engine/Plugins/Experimental/ChaosCaching/Source/ChaosCaching/Public/Chaos/Adapters/CacheAdapter.h b/Engine/Plugins/Experimental/ChaosCaching/Source/ChaosCaching/Public/Chaos/Adapters/CacheAdapter.h
--- a/Engine/Plugins/Experimental/ChaosCaching/Source/ChaosCaching/Public/Chaos/Adapters/CacheAdapter.h
+++ b/Engine/Plugins/Experimental/ChaosCaching/Source/ChaosCaching/Public/Chaos/Adapters/CacheAdapter.h
@@ -132,6 +132,11 @@ namespace Chaos
 		* Called from the game thread to set rest state from an evaluated time in the cache.
 		*/
 		virtual void SetRestState(UPrimitiveComponent* InComponent, UChaosCache* InCache, const FTransform& InRootTransform, Chaos::FReal InTime) const = 0;
+
+		/**
+		 * Called from the game thread to initialize a component and cache ready to load a cache at a given frame
+		 */
+		virtual void InitializeForLoad(UPrimitiveComponent* InComponent, FObservedComponent& InObserved) {}
 		
 		/**
 		 * Called from the game thread to initialize a component and cache ready to record a cache
diff --git a/Engine/Plugins/Experimental/ChaosCaching/Source/ChaosCaching/Public/Chaos/CacheCollection.h b/Engine/Plugins/Experimental/ChaosCaching/Source/ChaosCaching/Public/Chaos/CacheCollection.h
--- a/Engine/Plugins/Experimental/ChaosCaching/Source/ChaosCaching/Public/Chaos/CacheCollection.h
+++ b/Engine/Plugins/Experimental/ChaosCaching/Source/ChaosCaching/Public/Chaos/CacheCollection.h
@@ -28,6 +28,9 @@ public:
 	/** Return the max duration of all the caches stored in the collection */
 	CHAOSCACHING_API float GetMaxDuration() const;
 
+	/** Return the max number of frames of all the caches stored in the collection */
+	CHAOSCACHING_API uint32 GetMaxNumFrames() const;
+
 	UPROPERTY(EditAnywhere, Instanced, Category="Caching", meta=(EditFixedOrder))
 	TArray<TObjectPtr<UChaosCache>> Caches;
 };
diff --git a/Engine/Plugins/Experimental/ChaosCaching/Source/ChaosCaching/Public/Chaos/CacheManagerActor.h b/Engine/Plugins/Experimental/ChaosCaching/Source/ChaosCaching/Public/Chaos/CacheManagerActor.h
--- a/Engine/Plugins/Experimental/ChaosCaching/Source/ChaosCaching/Public/Chaos/CacheManagerActor.h
+++ b/Engine/Plugins/Experimental/ChaosCaching/Source/ChaosCaching/Public/Chaos/CacheManagerActor.h
@@ -267,6 +267,21 @@ public:
 
 	/** Accessor to the manager observed components (read/write) */
 	TArray<FObservedComponent>& GetObservedComponents() {return ObservedComponents;}
+	
+	/** Find or add a primitive component to a cache manager */
+	CHAOSCACHING_API FObservedComponent& FindOrAddObservedComponent(UPrimitiveComponent* InComponent, const FName& CacheName = TEXT(""), const bool bTransferSimulationFlag = false);
+
+	/** Remove a primitive component from the cache manager */
+	CHAOSCACHING_API void RemoveObservedComponent(UPrimitiveComponent* InComponent);
+	
+	/** Clear all the observed components */
+	CHAOSCACHING_API void ClearObservedComponents();
+	
+#if WITH_EDITOR
+	CHAOSCACHING_API void SetObservedComponentProperties(const ECacheMode& NewCacheMode);
+#endif
+	
+
 
 protected:
 
@@ -329,8 +344,6 @@ protected:
 
 	CHAOSCACHING_API FObservedComponent* FindObservedComponent(UPrimitiveComponent* InComponent);
 	CHAOSCACHING_API FObservedComponent& AddNewObservedComponent(UPrimitiveComponent* InComponent);
-	CHAOSCACHING_API FObservedComponent& FindOrAddObservedComponent(UPrimitiveComponent* InComponent);
-	CHAOSCACHING_API void ClearObservedComponents();
 
 	// Determines if the actor is allowed to record a cache.
 	bool bCanRecord;
@@ -339,9 +352,6 @@ protected:
 	bool bIsSimulating;
 
 private:
-#if WITH_EDITOR
-	CHAOSCACHING_API void SetObservedComponentProperties(const ECacheMode& NewCacheMode);
-#endif
 	
 	friend class UActorFactoryCacheManager; // Allows the actor factory to set up the observed list. See UActorFactoryCacheManager::PostSpawnActor
 
diff --git a/Engine/Plugins/Experimental/ChaosCaching/Source/ChaosCaching/Public/Chaos/ChaosCache.h b/Engine/Plugins/Experimental/ChaosCaching/Source/ChaosCaching/Public/Chaos/ChaosCache.h
--- a/Engine/Plugins/Experimental/ChaosCaching/Source/ChaosCaching/Public/Chaos/ChaosCache.h
+++ b/Engine/Plugins/Experimental/ChaosCaching/Source/ChaosCaching/Public/Chaos/ChaosCache.h
@@ -482,4 +482,10 @@ private:
 
 	/** Reverse Lookup for ChannelCurveToParticle. Rebuilt on load.*/
 	TMap<int32,int32> ParticleToChannelCurve;
+
+	/** Min time in case we are not writing to particle/curves/channels datas */
+	float MinTime = TNumericLimits<float>::Max();
+
+	/** Max time in case we are not writing to particle/curves/channels datas */
+	float MaxTime = TNumericLimits<float>::Lowest();
 };
diff --git a/Engine/Plugins/Experimental/ChaosCachingUSD/Source/ChaosCachingUSD/Private/ChaosCachingUSD/Operations.cpp b/Engine/Plugins/Experimental/ChaosCachingUSD/Source/ChaosCachingUSD/Private/ChaosCachingUSD/Operations.cpp
--- a/Engine/Plugins/Experimental/ChaosCachingUSD/Source/ChaosCachingUSD/Private/ChaosCachingUSD/Operations.cpp
+++ b/Engine/Plugins/Experimental/ChaosCachingUSD/Source/ChaosCachingUSD/Private/ChaosCachingUSD/Operations.cpp
@@ -624,6 +624,31 @@ UE::ChaosCachingUSD::WritePoints(
 	return WritePoints(Stage, PrimPath, Time, VtPoints, VtVels);
 }
 
+bool UE::ChaosCachingUSD::WritePoints(
+	UE::FUsdStage& Stage, 
+	const FString& PrimPath, 
+	const double Time, 
+	const TArray<Chaos::TVector<float, 3>>& Points, 
+	const TArray<Chaos::TVector<float, 3>>& Vels,
+	const FIntVector2& PointsRange)
+{
+	FScopedUsdAllocs UsdAllocs; // Use USD memory allocator
+
+	const size_t NumPoints = PointsRange[1];
+
+	pxr::VtArray<pxr::GfVec3f> VtPoints(NumPoints);
+	pxr::VtArray<pxr::GfVec3f> VtVels(NumPoints);
+
+	int32 GlobalIndex = PointsRange[0];
+	for (int32 LocalIndex = 0; LocalIndex < NumPoints; ++LocalIndex, ++GlobalIndex)
+	{
+		VtPoints[LocalIndex].Set(Points[GlobalIndex][0], Points[GlobalIndex][1], Points[GlobalIndex][2]);
+		VtVels[LocalIndex].Set(Vels[GlobalIndex][0], Vels[GlobalIndex][1], Vels[GlobalIndex][2]);
+	}
+
+	return WritePoints(Stage, PrimPath, Time, VtPoints, VtVels);
+}
+
 bool
 UE::ChaosCachingUSD::ReadTimeSamples(
 	const UE::FUsdStage& Stage, 
diff --git a/Engine/Plugins/Experimental/ChaosCachingUSD/Source/ChaosCachingUSD/Public/ChaosCachingUSD/Operations.h b/Engine/Plugins/Experimental/ChaosCachingUSD/Source/ChaosCachingUSD/Public/ChaosCachingUSD/Operations.h
--- a/Engine/Plugins/Experimental/ChaosCachingUSD/Source/ChaosCachingUSD/Public/ChaosCachingUSD/Operations.h
+++ b/Engine/Plugins/Experimental/ChaosCachingUSD/Source/ChaosCachingUSD/Public/ChaosCachingUSD/Operations.h
@@ -138,6 +138,10 @@ namespace UE::ChaosCachingUSD
 	/** Copy \p Points and \p Vels to \c VtArray (with USD memory allocator), then write to USD stage. */
 	CHAOSCACHINGUSD_API bool WritePoints(UE::FUsdStage& Stage, const FString& PrimPath, const double Time, const TArray<Chaos::TVector<float, 3>>& Points, const TArray<Chaos::TVector<float, 3>>& Vels);
 
+	/** Copy \p Points and \p Vels within the specified indices range to \c VtArray (with USD memory allocator), then write to USD stage. */
+	CHAOSCACHINGUSD_API bool WritePoints(UE::FUsdStage& Stage, const FString& PrimPath, const double Time, const TArray<Chaos::TVector<float, 3>>& Points, const TArray<Chaos::TVector<float, 3>>& Vels, const FIntVector2& PointsRange);
+	
+
 	/** Get time samples for an attribute. */
 	CHAOSCACHINGUSD_API bool ReadTimeSamples(const UE::FUsdStage& Stage, const FString& PrimPath, const FString& AttrName, TArray<double>& TimeSamples);
 	/** Get time samples for the points attribute. */
diff --git a/Engine/Plugins/Experimental/ChaosFlesh/Source/ChaosFleshEngine/Private/ChaosCache/FleshComponentCacheAdapter.cpp b/Engine/Plugins/Experimental/ChaosFlesh/Source/ChaosFleshEngine/Private/ChaosCache/FleshComponentCacheAdapter.cpp
--- a/Engine/Plugins/Experimental/ChaosFlesh/Source/ChaosFleshEngine/Private/ChaosCache/FleshComponentCacheAdapter.cpp
+++ b/Engine/Plugins/Experimental/ChaosFlesh/Source/ChaosFleshEngine/Private/ChaosCache/FleshComponentCacheAdapter.cpp
@@ -45,9 +45,12 @@ namespace Chaos
 	FFleshCacheAdapter::~FFleshCacheAdapter()
 	{
 #if USE_USD_SDK && DO_USD_CACHING
-		if (MonolithStage)
+		for(TPair<UPrimitiveComponent*,FPrimitiveStage>& PrimitiveStage : PrimitiveStages)
 		{
-			UE::ChaosCachingUSD::CloseStage(MonolithStage);
+			if(PrimitiveStage.Value.MonolithStage)
+			{
+				UE::ChaosCachingUSD::CloseStage(PrimitiveStage.Value.MonolithStage);
+			}
 		}
 #endif // USE_USD_SDK && DO_USD_CACHING
 	}
@@ -93,28 +96,31 @@ namespace Chaos
 #if USE_USD_SDK && DO_USD_CACHING
 					FScopedUsdAllocs UsdAllocs; // Use USD memory allocator
 
+					FPrimitiveStage& PrimitiveStage = PrimitiveStages.FindOrAdd(InComponent);
+					const FIntVector2 ParticleRange = GetParticleRange(InComponent, NumParticles);
+
 					// Update time range.
-					MinTime = FMath::Min(InTime, MinTime);
-					MaxTime = FMath::Max(InTime, MaxTime);
+					PrimitiveStage.MinTime = FMath::Min(InTime, PrimitiveStage.MinTime);
+					PrimitiveStage.MaxTime = FMath::Max(InTime, PrimitiveStage.MaxTime);
 
-					if (MonolithStage)
+					if (PrimitiveStage.MonolithStage)
 					{
-						if (!UE::ChaosCachingUSD::WritePoints(MonolithStage, PrimPath, InTime, Particles.XArray(), Particles.GetV()))
+						if (!UE::ChaosCachingUSD::WritePoints(PrimitiveStage.MonolithStage, PrimitiveStage.PrimPath, InTime, Particles.XArray(), Particles.GetV(), ParticleRange))
 						{
 							UE_LOG(LogChaosFleshCache, Error,
 								TEXT("Failed to write points '%s' at time %g to file: '%s'"),
-								*PrimPath, InTime, *MonolithStage.GetRootLayer().GetDisplayName());
+								*PrimitiveStage.PrimPath, InTime, *PrimitiveStage.MonolithStage.GetRootLayer().GetDisplayName());
 							return;
 						}
 
-						uint64 NumTimeSamples = UE::ChaosCachingUSD::GetNumTimeSamples(MonolithStage, PrimPath, UE::ChaosCachingUSD::GetPointsAttrName());
+						uint64 NumTimeSamples = UE::ChaosCachingUSD::GetNumTimeSamples(PrimitiveStage.MonolithStage, PrimitiveStage.PrimPath, UE::ChaosCachingUSD::GetPointsAttrName());
 						if (CVarParams.SaveFrequency >= 1 && NumTimeSamples % CVarParams.SaveFrequency == 0)
 						{
-							if (!UE::ChaosCachingUSD::SaveStage(MonolithStage, MinTime, MaxTime))
+							if (!UE::ChaosCachingUSD::SaveStage(PrimitiveStage.MonolithStage, PrimitiveStage.MinTime, PrimitiveStage.MaxTime))
 							{
 								UE_LOG(LogChaosFleshCache, Error,
 									TEXT("Failed to save file: '%s'"),
-									*MonolithStage.GetRootLayer().GetDisplayName());
+									*PrimitiveStage.MonolithStage.GetRootLayer().GetDisplayName());
 							}
 						}
 					}
@@ -162,115 +168,126 @@ namespace Chaos
 		}
 	}
 
-	void FFleshCacheAdapter::Playback_PreSolve(UPrimitiveComponent* InComponent, UChaosCache* InCache, Chaos::FReal InTime, FPlaybackTickRecord& TickRecord, TArray<TPBDRigidParticleHandle<Chaos::FReal, 3>*>& OutUpdatedRigids) const
+	void FFleshCacheAdapter::LoadCacheAtTime(UPrimitiveComponent* PrimitiveComponent, Chaos::FReal TargetTime, const int32 NumParticles, const bool bNeedsRange,
+		const TFunction<void(const int32, const int32, const FVector3f&, const FVector3f&)>& LoadFunction) const
 	{
-		if (FDeformableSolver* Solver = GetDeformableSolver(InComponent))
+#if USE_USD_SDK && DO_USD_CACHING
+
+		const FPrimitiveStage& PrimitiveStage = PrimitiveStages.FindOrAdd(PrimitiveComponent);
+		if (PrimitiveStage.MonolithStage)
 		{
-			FDeformableSolver::FPhysicsThreadAccess PhysicsThreadAccess(Solver, Softs::FPhysicsThreadAccessor());
+			// GetParticleRange only called in play/record mode (not load), so PT only
+			const FIntVector2 ParticleRange = bNeedsRange ? GetParticleRange(PrimitiveComponent, NumParticles) : FIntVector2(0,NumParticles);
+			const int32 NumComponentParticles = ParticleRange[1];
+			
+			FScopedUsdAllocs UEAllocs; // Use USD memory allocator
+			pxr::VtArray<pxr::GfVec3f> Points0, Points1;
+			pxr::VtArray<pxr::GfVec3f> Vels0, Vels1;
+
+			double Prev = -TNumericLimits<double>::Max();
+			double Next = -TNumericLimits<double>::Max();
+			double PrevV = -TNumericLimits<double>::Max();
+			double NextV = -TNumericLimits<double>::Max();
+			
+			if (!UE::ChaosCachingUSD::GetBracketingTimeSamples(
+					PrimitiveStage.MonolithStage, PrimitiveStage.PrimPath, UE::ChaosCachingUSD::GetPointsAttrName(), TargetTime, &Prev, &Next) ||
+				!UE::ChaosCachingUSD::GetBracketingTimeSamples(
+					PrimitiveStage.MonolithStage, PrimitiveStage.PrimPath, UE::ChaosCachingUSD::GetVelocityAttrName(), TargetTime, &PrevV, &NextV) ||
+				Prev != PrevV ||
+				Next != NextV)
+			{
+				UE_LOG(LogChaosFleshCache, Error,
+					TEXT("Inconsistent bracketing time samples for attributes '%s' and '%s' at frame %g from file: '%s'"),
+					*UE::ChaosCachingUSD::GetPointsAttrName(),
+					*UE::ChaosCachingUSD::GetVelocityAttrName(),
+					TargetTime, *PrimitiveStage.MonolithStage.GetRootLayer().GetDisplayName());
+				return;
+			}
 
-			if (FEvolution* Evolution = PhysicsThreadAccess.GetEvolution())
+			if (!UE::ChaosCachingUSD::ReadPoints(PrimitiveStage.MonolithStage, PrimitiveStage.PrimPath, Prev, Points0, Vels0) ||
+				Points0.size() != Vels0.size())
 			{
-#if USE_USD_SDK && DO_USD_CACHING
-				if (MonolithStage)
+				UE_LOG(LogChaosFleshCache, Error,
+					TEXT("Failed to read points '%s' at time %g from file: '%s'"),
+					*PrimitiveStage.PrimPath, Prev, *PrimitiveStage.MonolithStage.GetRootLayer().GetDisplayName());
+				return;
+			}
+
+			int32 NumCachedParticles = static_cast<int32>(Points0.size());
+			if (NumCachedParticles > NumComponentParticles)
+			{
+				// Cached particles doesn't match solver particles.  Truncate.
+				NumCachedParticles = NumComponentParticles;
+			}
+
+			// < time range start, > time range end, or exact hit
+			if (FMath::IsNearlyEqual(Prev, Next))
+			{
+				// Directly set the result of the cache into the solver particles
+				for (int32 CachedIndex = 0; CachedIndex < NumCachedParticles; ++CachedIndex)
 				{
-					FScopedUsdAllocs UEAllocs; // Use USD memory allocator
-					pxr::VtArray<pxr::GfVec3f> Points0, Points1;
-					pxr::VtArray<pxr::GfVec3f> Vels0, Vels1;
-
-					double TargetTime = InTime;// TickRecord.GetTime();
-					double Prev = -TNumericLimits<double>::Max();
-					double Next = -TNumericLimits<double>::Max();
-					double PrevV = -TNumericLimits<double>::Max();
-					double NextV = -TNumericLimits<double>::Max();
-					if (!UE::ChaosCachingUSD::GetBracketingTimeSamples(
-							MonolithStage, PrimPath, UE::ChaosCachingUSD::GetPointsAttrName(), TargetTime, &Prev, &Next) ||
-						!UE::ChaosCachingUSD::GetBracketingTimeSamples(
-							MonolithStage, PrimPath, UE::ChaosCachingUSD::GetVelocityAttrName(), TargetTime, &PrevV, &NextV) ||
-						Prev != PrevV ||
-						Next != NextV)
-					{
-						UE_LOG(LogChaosFleshCache, Error,
-							TEXT("Inconsistent bracketing time samples for attributes '%s' and '%s' at frame %g from file: '%s'"),
-							*UE::ChaosCachingUSD::GetPointsAttrName(),
-							*UE::ChaosCachingUSD::GetVelocityAttrName(),
-							TargetTime, *MonolithStage.GetRootLayer().GetDisplayName());
-						return;
-					}
+					// Note that VtArray::operator[] is non-const access and will cause trigger 
+					// the copy-on-write memcopy!  VtArray::cdata() avoids that.
+					const pxr::GfVec3f& Pos = Points0.cdata()[CachedIndex];
+					const pxr::GfVec3f& Vel = Vels0.cdata()[CachedIndex];
 
-					if (!UE::ChaosCachingUSD::ReadPoints(MonolithStage, PrimPath, Prev, Points0, Vels0) ||
-						Points0.size() != Vels0.size())
-					{
-						UE_LOG(LogChaosFleshCache, Error,
-							TEXT("Failed to read points '%s' at time %g from file: '%s'"),
-							*PrimPath, Prev, *MonolithStage.GetRootLayer().GetDisplayName());
-						return;
-					}
+					LoadFunction(CachedIndex, ParticleRange[0], FVector3f(Pos[0], Pos[1], Pos[2]), FVector3f(Vel[0], Vel[1], Vel[2]));
+				}
+				return;
+			}
 
-					FParticles& Particles = Evolution->Particles();
-					const int32 NumParticles = Particles.Size();
-					Softs::FSolverVec3* ParticleXs = &Particles.X(0);
-					Softs::FSolverVec3* ParticleVs = &Particles.V(0);
-					Softs::FSolverReal* ParticleInvMs = &Particles.InvM(0);
-					Softs::FPAndInvM* ParticlePAndInvMs = &Particles.PAndInvM(0);
+			if (!UE::ChaosCachingUSD::ReadPoints(PrimitiveStage.MonolithStage, PrimitiveStage.PrimPath, Next, Points1, Vels1) ||
+				Points1.size() != Vels1.size() ||
+				Points0.size() != Points1.size())
+			{
+				UE_LOG(LogChaosFleshCache, Error,
+					TEXT("Failed to read points '%s' at time %g from file: '%s'"),
+					*PrimitiveStage.PrimPath, Next, *PrimitiveStage.MonolithStage.GetRootLayer().GetDisplayName());
+				return;
+			}
+			double Duration = Next - Prev;
+			double Alpha = Duration > UE_SMALL_NUMBER ? (TargetTime - Prev) / Duration : 0.5;
+			for (int32 CachedIndex = 0; CachedIndex < NumCachedParticles; ++CachedIndex)
+			{
+				// Note that VtArray::operator[] is non-const access and will cause trigger 
+				// the copy-on-write memcopy!  VtArray::cdata() avoids that.
+				const pxr::GfVec3f& P0 = Points0.cdata()[CachedIndex];
+				const pxr::GfVec3f& P1 = Points1.cdata()[CachedIndex];
+				const pxr::GfVec3f& V0 = Vels0.cdata()[CachedIndex];
+				const pxr::GfVec3f& V1 = Vels1.cdata()[CachedIndex];
+				pxr::GfVec3f Pos = (1.0 - Alpha) * P0 + Alpha * P1;
+				pxr::GfVec3f Vel = (1.0 - Alpha) * V0 + Alpha * V1;
+				
+				LoadFunction(CachedIndex, ParticleRange[0], FVector3f(Pos[0], Pos[1], Pos[2]), FVector3f(Vel[0], Vel[1], Vel[2]));
+			}
+		}
+	#endif
+	}
 
-					if (UFleshComponent* FleshComp = CastChecked<UFleshComponent>(InComponent))
-					{
-						if (Chaos::Softs::FFleshThreadingProxy* FleshProxy = FleshComp->GetPhysicsProxy()->As<Chaos::Softs::FFleshThreadingProxy>())
-						{
-							FIntVector2 ParticleRange = FleshProxy->GetSolverParticleRange();
+	void FFleshCacheAdapter::Playback_PreSolve(UPrimitiveComponent* InComponent, UChaosCache* InCache, Chaos::FReal InTime, FPlaybackTickRecord& TickRecord, TArray<TPBDRigidParticleHandle<Chaos::FReal, 3>*>& OutUpdatedRigids) const
+	{
+		if (FDeformableSolver* Solver = GetDeformableSolver(InComponent))
+		{
+			FDeformableSolver::FPhysicsThreadAccess PhysicsThreadAccess(Solver, Softs::FPhysicsThreadAccessor());
 
-							int32 NumCachedParticles = static_cast<int32>(Points0.size());
-							if (NumCachedParticles > ParticleRange[1])
-							{
-								// Cached particles doesn't match solver particles.  Truncate.
-								NumCachedParticles = ParticleRange[1];
-							}
-							// < time range start, > time range end, or exact hit
-							if (FMath::IsNearlyEqual(Prev, Next))
-							{
-								// Directly set the result of the cache into the solver particles
-								for (int32 CachedIndex = ParticleRange[0]; CachedIndex < ParticleRange[0] + ParticleRange[1]; ++CachedIndex)
-								{
-									// Note that VtArray::operator[] is non-const access and will cause trigger 
-									// the copy-on-write memcopy!  VtArray::cdata() avoids that.
-									const pxr::GfVec3f& P0 = Points0.cdata()[CachedIndex - ParticleRange[0]];
-									const pxr::GfVec3f& V0 = Vels0.cdata()[CachedIndex - ParticleRange[0]];
-									ParticleXs[CachedIndex].Set(P0[0], P0[1], P0[2]);
-									ParticleVs[CachedIndex].Set(V0[0], V0[1], V0[2]);
-									ParticleInvMs[CachedIndex] = Softs::FSolverReal(0);
-									ParticlePAndInvMs[CachedIndex].InvM = Softs::FSolverReal(0);
-									ParticlePAndInvMs[CachedIndex].P = ParticleXs[CachedIndex];
-								}
-								return;
-							}
+			if (FEvolution* Evolution = PhysicsThreadAccess.GetEvolution())
+			{
+#if USE_USD_SDK && DO_USD_CACHING
 
-							if (!UE::ChaosCachingUSD::ReadPoints(MonolithStage, PrimPath, Next, Points1, Vels1) ||
-								Points1.size() != Vels1.size() ||
-								Points0.size() != Points1.size())
-							{
-								UE_LOG(LogChaosFleshCache, Error,
-									TEXT("Failed to read points '%s' at time %g from file: '%s'"),
-									*PrimPath, Next, *MonolithStage.GetRootLayer().GetDisplayName());
-								return;
-							}
-							double Duration = Next - Prev;
-							double Alpha = Duration > UE_SMALL_NUMBER ? (TargetTime - Prev) / Duration : 0.5;
-							for (int32 CachedIndex = ParticleRange[0]; CachedIndex < ParticleRange[0] + ParticleRange[1]; ++CachedIndex)
-							{
-								// Note that VtArray::operator[] is non-const access and will cause trigger 
-								// the copy-on-write memcopy!  VtArray::cdata() avoids that.
-								const pxr::GfVec3f& P0 = Points0.cdata()[CachedIndex - ParticleRange[0]];
-								const pxr::GfVec3f& P1 = Points1.cdata()[CachedIndex - ParticleRange[0]];
-								const pxr::GfVec3f& V0 = Vels0.cdata()[CachedIndex - ParticleRange[0]];
-								const pxr::GfVec3f& V1 = Vels1.cdata()[CachedIndex - ParticleRange[0]];
-								pxr::GfVec3f Pos = (1.0 - Alpha) * P0 + Alpha * P1;
-								pxr::GfVec3f Vel = (1.0 - Alpha) * V0 + Alpha * V1;
-								ParticleXs[CachedIndex].Set(Pos[0], Pos[1], Pos[2]);
-								ParticleVs[CachedIndex].Set(Vel[0], Vel[1], Vel[2]);
-							}
-						}
-					}
-				}
+				FParticles& Particles = Evolution->Particles();
+				const int32 NumParticles = Particles.Size();
+				Softs::FSolverVec3* ParticleXs = &Particles.X(0);
+				Softs::FSolverVec3* ParticleVs = &Particles.V(0);
+				
+				LoadCacheAtTime(InComponent, InTime, NumParticles, true, [&ParticleXs, &ParticleVs](
+					const int32 ParticleIndex, const int32 ParticleOffset, const FVector3f& ParticlePosition, const FVector3f& ParticleVelocity)
+				{
+					const int32 GlobalIndex = ParticleIndex + ParticleOffset;
+					
+					ParticleXs[GlobalIndex].Set(ParticlePosition[0], ParticlePosition[1], ParticlePosition[2]);
+					ParticleVs[GlobalIndex].Set(ParticleVelocity[0], ParticleVelocity[1], ParticleVelocity[2]);
+				});
+				
 #else // USE_USD_SDK && DO_USD_CACHING
 				FCacheEvaluationContext Context(TickRecord);
 				Context.bEvaluateTransform = false;
@@ -340,7 +357,7 @@ namespace Chaos
 #endif // USE_USD_SDK && DO_USD_CACHING
 	}
 
-	Chaos::Softs::FDeformableSolver* FFleshCacheAdapter::GetDeformableSolver(UPrimitiveComponent* InComponent) const
+	Chaos::Softs::FDeformableSolver* FFleshCacheAdapter::GetDeformableSolver(UPrimitiveComponent* InComponent)
 	{
 		if(InComponent)
 		{
@@ -355,6 +372,25 @@ namespace Chaos
 		return nullptr;
 	}
 
+	FIntVector2 FFleshCacheAdapter::GetParticleRange(UPrimitiveComponent* InComponent, const int32 NumParticles)
+	{
+		FIntVector2 ParticleRange(0, NumParticles);
+		if(InComponent)
+		{
+			if(UFleshComponent* FleshComp = CastChecked<UFleshComponent>(InComponent))
+			{
+				if(FleshComp->GetPhysicsProxy())
+				{
+					if (const Chaos::Softs::FFleshThreadingProxy* Proxy = FleshComp->GetPhysicsProxy()->As<Chaos::Softs::FFleshThreadingProxy>())
+					{
+						ParticleRange = Proxy->GetSolverParticleRange();
+					}
+				}
+			}
+		}
+		return ParticleRange;
+	}
+
 	Chaos::FPhysicsSolver* FFleshCacheAdapter::GetComponentSolver(UPrimitiveComponent* InComponent) const
 	{
 		return nullptr;
@@ -370,6 +406,30 @@ namespace Chaos
 		}
 		return nullptr;
 	}
+
+	void FFleshCacheAdapter::InitializeForLoad(UPrimitiveComponent* InComponent, FObservedComponent& InObserved)
+	{
+#if USE_USD_SDK && DO_USD_CACHING
+		if(const UFleshComponent* FleshComp = CastChecked<UFleshComponent>(InComponent))
+		{
+			FPrimitiveStage& PrimitiveStage = PrimitiveStages.FindOrAdd(InComponent);
+			
+			PrimitiveStage.PrimPath = UsdUtils::GetPrimPathForObject(FleshComp);
+			PrimitiveStage.FilePath = GetUSDCacheFilePathRO(InObserved, FleshComp);
+			
+			if (!PrimitiveStage.MonolithStage)
+			{
+				FPlatformFileManager& FileManager = FPlatformFileManager::Get();
+				IPlatformFile& PlatformFile = FileManager.GetPlatformFile();
+				
+				if (PlatformFile.FileExists(*PrimitiveStage.FilePath))
+				{
+					UE::ChaosCachingUSD::OpenStage(PrimitiveStage.FilePath, PrimitiveStage.MonolithStage);
+				}
+			}
+		}
+#endif // USE_USD_SDK && DO_USD_CACHING
+	}
 	
 	void FFleshCacheAdapter::SetRestState(UPrimitiveComponent* InComponent, UChaosCache* InCache, const FTransform& InRootTransform, Chaos::FReal InTime) const
 	{
@@ -383,47 +443,29 @@ namespace Chaos
 		if (UFleshComponent* FleshComp = CastChecked<UFleshComponent>(InComponent))
 		{
 #if USE_USD_SDK && DO_USD_CACHING
+			
 			FleshComp->ResetDynamicCollection();
 			if (UFleshDynamicAsset* DynamicCollection = FleshComp->GetDynamicCollection())
 			{
 				TManagedArray<FVector3f>& DynamicVertex = DynamicCollection->GetPositions();
-				//TManagedArray<FVector3f>& DynamicVertex = DynamicCollection->GetVelocities();
 				const int32 NumDynamicVertex = DynamicVertex.Num();
 
-				if (MonolithStage)
+				auto UEVertd = [](Chaos::FVec3 V) { return FVector3d(V.X, V.Y, V.Z); };
+				auto UEVertf = [](FVector3d V) { return FVector3f((float)V.X, (float)V.Y, (float)V.Z); };
+				
+				LoadCacheAtTime(InComponent, InTime, NumDynamicVertex, false, [&DynamicVertex, FleshComp, UEVertd, UEVertf](
+					const int32 ParticleIndex, const int32 ParticleOffset, const FVector3f& ParticlePosition, const FVector3f&)
 				{
-					FScopedUsdAllocs UEAllocs; // Use USD memory allocator
-					pxr::VtArray<pxr::GfVec3f> Points;
-					pxr::VtArray<pxr::GfVec3f> Vels;
+					const int32 GlobalIndex = ParticleIndex + ParticleOffset;
 
-					if (!UE::ChaosCachingUSD::ReadPoints(MonolithStage, PrimPath, -TNumericLimits<double>::Max(), Points, Vels))
-					{
-						UE_LOG(LogChaosFleshCache, Error, 
-							TEXT("Failed to read points '%s' at time 'default' from file: '%s'"), 
-							*PrimPath, *MonolithStage.GetRootLayer().GetDisplayName());
-						return;
-					}
-
-					int32 NumCachedParticles = static_cast<int32>(Points.size());
-					if (NumDynamicVertex == NumCachedParticles)
-					{
-						for (int32 CachedIndex = 0; CachedIndex < NumCachedParticles; ++CachedIndex)
-						{
-							// Note that VtArray::operator[] is non-const access and will cause trigger 
-							// the copy-on-write memcopy!  VtArray::cdata() avoids that.
-							const pxr::GfVec3f& P = Points.cdata()[CachedIndex];
-							DynamicVertex[CachedIndex].Set(P[0], P[1], P[2]);
-
-							auto UEVertd = [](Chaos::FVec3 V) { return FVector3d(V.X, V.Y, V.Z); };
-							auto UEVertf = [](FVector3d V) { return FVector3f((float)V.X, (float)V.Y, (float)V.Z); };
-							DynamicVertex[CachedIndex] =
-								UEVertf(FleshComp->GetComponentTransform().InverseTransformPosition(
-									UEVertd(DynamicVertex[CachedIndex])));
-						}
-					}
-				}
+					DynamicVertex[GlobalIndex].Set(ParticlePosition[0], ParticlePosition[1], ParticlePosition[2]);
+					DynamicVertex[GlobalIndex] =
+									UEVertf(FleshComp->GetComponentTransform().InverseTransformPosition(
+										UEVertd(DynamicVertex[GlobalIndex])));
+				});
 			}
 
+			
 #else // USE_USD_SDK && DO_USD_CACHING
 
 			FPlaybackTickRecord TickRecord;
@@ -566,6 +608,7 @@ namespace Chaos
 		FString FileName = FString::Printf(TEXT("%s.%s"), *CompName, *Ext);
 		FString FilePath = FPaths::Combine(CacheDir, FileName);
 
+		
 		return FilePath;
 	}
 #endif // USE_USD_SDK && DO_USD_CACHING
@@ -606,16 +649,18 @@ namespace Chaos
 					return false;
 				}
 
-				PrimPath = UsdUtils::GetPrimPathForObject(FleshComp);
+				FPrimitiveStage& PrimitiveStage = PrimitiveStages.FindOrAdd(InComponent);
+
+				PrimitiveStage.PrimPath = UsdUtils::GetPrimPathForObject(FleshComp);
 				if (bUseMonolith)
 				{
 					bReadOnly = false;
 					FString Ext = CVarParams.bWriteBinary ? FString(TEXT("usd")) : FString(TEXT("usda"));
 					FString FileName = FString::Printf(TEXT("%s.%s"), *CompName, *Ext);
-					FilePath = FPaths::Combine(CacheDir, FileName);
+					PrimitiveStage.FilePath = FPaths::Combine(CacheDir, FileName);
 					if (CVarParams.bNoClobber)
 					{
-						if (PlatformFile.FileExists(*FilePath))
+						if (PlatformFile.FileExists(*PrimitiveStage.FilePath))
 						{
 							// Rename the file to 'path/to/file.usd' to 'path/to/file_#.usd', where '#' 
 							// is a unique version number.
@@ -626,34 +671,34 @@ namespace Chaos
 								UniqueFilePath = FPaths::Combine(CacheDir, UniqueCompName);
 							} while (PlatformFile.FileExists(*UniqueFilePath));
 
-							if (!PlatformFile.MoveFile(*UniqueFilePath, *FilePath))
+							if (!PlatformFile.MoveFile(*UniqueFilePath, *PrimitiveStage.FilePath))
 							{
-								UE_LOG(LogChaosFleshCache, Error, TEXT("Failed to rename file from '%s' to '%s'."), *FilePath, *UniqueFilePath);
+								UE_LOG(LogChaosFleshCache, Error, TEXT("Failed to rename file from '%s' to '%s'."), *PrimitiveStage.FilePath, *UniqueFilePath);
 								return false;
 							}
 						}
 					}
 					else
 					{
-						if (!PlatformFile.DeleteFile(*FilePath))
+						if (!PlatformFile.DeleteFile(*PrimitiveStage.FilePath))
 						{
-							UE_LOG(LogChaosFleshCache, Error, TEXT("Failed to remove existing cache file: '%s'"), *FilePath);
+							UE_LOG(LogChaosFleshCache, Error, TEXT("Failed to remove existing cache file: '%s'"), *PrimitiveStage.FilePath);
 							return false;
 						}
 					}
 
-					if (MonolithStage)
+					if (PrimitiveStage.MonolithStage)
 					{
-						UE::ChaosCachingUSD::CloseStage(MonolithStage);
+						UE::ChaosCachingUSD::CloseStage(PrimitiveStage.MonolithStage);
 					}
-					if (!UE::ChaosCachingUSD::NewStage(FilePath, MonolithStage))
+					if (!UE::ChaosCachingUSD::NewStage(PrimitiveStage.FilePath, PrimitiveStage.MonolithStage))
 					{
-						UE_LOG(LogChaosFleshCache, Error, TEXT("Failed to create new USD file: '%s'"), *FilePath);
+						UE_LOG(LogChaosFleshCache, Error, TEXT("Failed to create new USD file: '%s'"), *PrimitiveStage.FilePath);
 						return false;
 					}
-					if (!UE::ChaosCachingUSD::WriteTetMesh(MonolithStage, PrimPath, *RestCollection))
+					if (!UE::ChaosCachingUSD::WriteTetMesh(PrimitiveStage.MonolithStage, PrimitiveStage.PrimPath, *RestCollection))
 					{
-						UE_LOG(LogChaosFleshCache, Error, TEXT("Failed to write tetrahedron mesh '%s' to USD file: '%s'"), *PrimPath, *FilePath);
+						UE_LOG(LogChaosFleshCache, Error, TEXT("Failed to write tetrahedron mesh '%s' to USD file: '%s'"), *PrimitiveStage.PrimPath, *PrimitiveStage.FilePath);
 						return false;
 					}
 				}
@@ -685,29 +730,31 @@ namespace Chaos
 				// USD caching
 				//
 
-				PrimPath = UsdUtils::GetPrimPathForObject(FleshComp);
+				FPrimitiveStage& PrimitiveStage = PrimitiveStages.FindOrAdd(InComponent);
+
+				PrimitiveStage.PrimPath = UsdUtils::GetPrimPathForObject(FleshComp);
 				if (bUseMonolith)
 				{
 					bReadOnly = true;
-					FilePath = GetUSDCacheFilePathRO(InObserved, FleshComp);
+					PrimitiveStage.FilePath = GetUSDCacheFilePathRO(InObserved, FleshComp);
 
 					FPlatformFileManager& FileManager = FPlatformFileManager::Get();
 					IPlatformFile& PlatformFile = FileManager.GetPlatformFile();
-					if (PlatformFile.FileExists(*FilePath))
+					if (PlatformFile.FileExists(*PrimitiveStage.FilePath))
 					{
-						if (MonolithStage)
+						if (PrimitiveStage.MonolithStage)
 						{
-							UE::ChaosCachingUSD::CloseStage(MonolithStage);
+							UE::ChaosCachingUSD::CloseStage(PrimitiveStage.MonolithStage);
 						}
-						if (!UE::ChaosCachingUSD::OpenStage(FilePath, MonolithStage))
+						if (!UE::ChaosCachingUSD::OpenStage(PrimitiveStage.FilePath, PrimitiveStage.MonolithStage))
 						{
-							UE_LOG(LogChaosFleshCache, Error, TEXT("Failed to open USD cache file: '%s'"), *FilePath);
+							UE_LOG(LogChaosFleshCache, Error, TEXT("Failed to open USD cache file: '%s'"), *PrimitiveStage.FilePath);
 							return false;
 						}
 					}
 					else
 					{
-						UE_LOG(LogChaosFleshCache, Error, TEXT("USD cache file not found: '%s'"), *FilePath);
+						UE_LOG(LogChaosFleshCache, Error, TEXT("USD cache file not found: '%s'"), *PrimitiveStage.FilePath);
 						return false;
 					}
 				}
@@ -725,21 +772,26 @@ namespace Chaos
 	void FFleshCacheAdapter::Finalize()
 	{
 #if USE_USD_SDK && DO_USD_CACHING
-		// Detach shared memory arrays.
-		if (MonolithStage)
+
+		for(TPair<UPrimitiveComponent*,FPrimitiveStage>& PrimitiveStage : PrimitiveStages)
 		{
-			if (!bReadOnly)
+			// Detach shared memory arrays.
+			if (PrimitiveStage.Value.MonolithStage)
 			{
-				UE::ChaosCachingUSD::SaveStage(MonolithStage, MinTime, MaxTime);
+				if (!bReadOnly)
+				{
+					UE::ChaosCachingUSD::SaveStage(PrimitiveStage.Value.MonolithStage, PrimitiveStage.Value.MinTime, PrimitiveStage.Value.MaxTime);
+				}
+				UE::ChaosCachingUSD::CloseStage(PrimitiveStage.Value.MonolithStage);
+				PrimitiveStage.Value.MonolithStage = UE::FUsdStage();
 			}
-			UE::ChaosCachingUSD::CloseStage(MonolithStage);
-			MonolithStage = UE::FUsdStage();
-		}
 
-		PrimPath.Empty();
-		FilePath.Empty();
-		MinTime = TNumericLimits<double>::Max();
-		MaxTime = -TNumericLimits<double>::Max();
+			PrimitiveStage.Value.PrimPath.Empty();
+			PrimitiveStage.Value.FilePath.Empty();
+			PrimitiveStage.Value.MinTime = TNumericLimits<double>::Max();
+			PrimitiveStage.Value.MaxTime = TNumericLimits<double>::Lowest();
+		}
+		PrimitiveStages.Reset();
 #endif // USE_USD_SDK && DO_USD_CACHING
 	}
 
diff --git a/Engine/Plugins/Experimental/ChaosFlesh/Source/ChaosFleshEngine/Public/ChaosCache/FleshComponentCacheAdapter.h b/Engine/Plugins/Experimental/ChaosFlesh/Source/ChaosFleshEngine/Public/ChaosCache/FleshComponentCacheAdapter.h
--- a/Engine/Plugins/Experimental/ChaosFlesh/Source/ChaosFleshEngine/Public/ChaosCache/FleshComponentCacheAdapter.h
+++ b/Engine/Plugins/Experimental/ChaosFlesh/Source/ChaosFleshEngine/Public/ChaosCache/FleshComponentCacheAdapter.h
@@ -67,6 +67,7 @@ namespace Chaos
 		virtual void						SetRestState(UPrimitiveComponent* InComponent, UChaosCache* InCache, const FTransform& InRootTransform, Chaos::FReal InTime) const override;
 		virtual bool						InitializeForRecord(UPrimitiveComponent* InComponent, FObservedComponent& InObserved) override;
 		virtual bool						InitializeForPlayback(UPrimitiveComponent* InComponent, FObservedComponent& InObserved, float InTime) override;
+		virtual void						InitializeForLoad(UPrimitiveComponent* InComponent, FObservedComponent& InObserved) override;
 		virtual void						Record_PostSolve(UPrimitiveComponent* InComp, const FTransform& InRootTransform, FPendingFrameWrite& OutFrame, Chaos::FReal InTime) const override;
 		virtual void						Playback_PreSolve(UPrimitiveComponent* InComponent,
 												UChaosCache*										InCache,
@@ -76,27 +77,45 @@ namespace Chaos
 		virtual void						Finalize() override;
 		// ~End FComponentCacheAdapter interface
 
+		/** Get the component deformable solver */
+		static FDeformableSolver* GetDeformableSolver(UPrimitiveComponent* InComponent);
 
-		FDeformableSolver* GetDeformableSolver(UPrimitiveComponent* InComponent) const;
-
-
+		/** Get the particle indices range for the component in the evolution particles list on the physics thread */
+		static FIntVector2 GetParticleRange(UPrimitiveComponent* InComponent, const int32 NumParticles);
+		
 	private:
+
+		/** Load the cache for a given component at a specific time */
+		void LoadCacheAtTime(UPrimitiveComponent* PrimitiveComponent, Chaos::FReal TargetTime, const int32 NumParticles, const bool bNeedsRange,
+				const TFunction<void(const int32, const int32, const FVector3f&, const FVector3f&)>& LoadFunction) const;
+		
 #if USE_USD_SDK && DO_USD_CACHING
 
 		bool bReadOnly = false;
 		bool bUseMonolith = true; // move to cvarparams when value clips is an option
 
-		/**
-		 * Cache files are stored in a "SimCache" directory at the root of the project.
-		 * File name is derived from flesh component name.
-		 */
-		FString FilePath;
-		FString PrimPath;
+		/** Structure holding the per primitive stage datas */
+		struct FPrimitiveStage
+		{
+			/**
+			 * Cache files are stored in a "SimCache" directory at the root of the project.
+			 * File name is derived from flesh component name.
+			 */
+			FString FilePath;
+			FString PrimPath;
+
+			/** Min time in the usd stage */
+			double MinTime = TNumericLimits<double>::Max();
+
+			/** Max time in the usd stage */
+			double MaxTime = TNumericLimits<double>::Lowest();
 
-		mutable double MinTime = TNumericLimits<double>::Max();
-		mutable double MaxTime = -TNumericLimits<double>::Max();
+			/** USD Monolith stage used to load/record datas */
+			UE::FUsdStage MonolithStage;
+		};
 
-		mutable UE::FUsdStage MonolithStage;
+		/** Per primitive stages datas */
+		mutable TMap<UPrimitiveComponent*,FPrimitiveStage> PrimitiveStages;
 
 #else // USE_USD_SDK && DO_USD_CACHING
 
