5- Flesh dataflow simulation nodes

diff --git a/Engine/Plugins/Experimental/ChaosFlesh/Source/ChaosFleshEngine/Private/ChaosFlesh/ChaosDeformablePhysicsComponent.cpp b/Engine/Plugins/Experimental/ChaosFlesh/Source/ChaosFleshEngine/Private/ChaosFlesh/ChaosDeformablePhysicsComponent.cpp
--- a/Engine/Plugins/Experimental/ChaosFlesh/Source/ChaosFleshEngine/Private/ChaosFlesh/ChaosDeformablePhysicsComponent.cpp
+++ b/Engine/Plugins/Experimental/ChaosFlesh/Source/ChaosFleshEngine/Private/ChaosFlesh/ChaosDeformablePhysicsComponent.cpp
@@ -114,11 +114,11 @@ void UDeformablePhysicsComponent::RemoveProxy(Chaos::Softs::FDeformableSolver::F
 
 UDeformableSolverComponent* UDeformablePhysicsComponent::GetDeformableSolver()
 {
-	return PrimarySolverComponent;
+	return PrimarySolverComponent.Get();
 }
 const UDeformableSolverComponent* UDeformablePhysicsComponent::GetDeformableSolver() const
 { 
-	return PrimarySolverComponent;
+	return PrimarySolverComponent.Get();
 }
 
 UDeformablePhysicsComponent::UDeformablePhysicsComponent(const FObjectInitializer& ObjectInitializer)
@@ -126,6 +126,18 @@ UDeformablePhysicsComponent::UDeformablePhysicsComponent(const FObjectInitialize
 {
 }
 
+void UDeformablePhysicsComponent::DisableSimulation()
+{
+	if (PrimarySolverComponent)
+	{
+		if (!PrimarySolverComponent->ConnectedObjects.DeformableComponents.Contains(this))
+		{
+			PrimarySolverComponent->ConnectedObjects.DeformableComponents.Remove(this);
+		}
+		PrimarySolverComponent->RemoveDeformableProxy(this);
+	}
+	PrimarySolverComponent = nullptr;
+}
 
 void UDeformablePhysicsComponent::EnableSimulation(UDeformableSolverComponent* DeformableSolverComponent)
 {
diff --git a/Engine/Plugins/Experimental/ChaosFlesh/Source/ChaosFleshEngine/Private/ChaosFlesh/ChaosDeformableSolverComponent.cpp b/Engine/Plugins/Experimental/ChaosFlesh/Source/ChaosFleshEngine/Private/ChaosFlesh/ChaosDeformableSolverComponent.cpp
--- a/Engine/Plugins/Experimental/ChaosFlesh/Source/ChaosFleshEngine/Private/ChaosFlesh/ChaosDeformableSolverComponent.cpp
+++ b/Engine/Plugins/Experimental/ChaosFlesh/Source/ChaosFleshEngine/Private/ChaosFlesh/ChaosDeformableSolverComponent.cpp
@@ -14,6 +14,7 @@ DECLARE_CYCLE_STAT(TEXT("Chaos.Deformable.UDeformableSolverComponent.UpdateDefor
 DECLARE_CYCLE_STAT(TEXT("Chaos.Deformable.UDeformableSolverComponent.TickComponent"), STAT_ChaosDeformable_UDeformableSolverComponent_TickComponent, STATGROUP_Chaos);
 DECLARE_CYCLE_STAT(TEXT("Chaos.Deformable.UDeformableSolverComponent.Reset"), STAT_ChaosDeformable_UDeformableSolverComponent_Reset, STATGROUP_Chaos);
 DECLARE_CYCLE_STAT(TEXT("Chaos.Deformable.UDeformableSolverComponent.AddDeformableProxy"), STAT_ChaosDeformable_UDeformableSolverComponent_AddDeformableProxy, STATGROUP_Chaos);
+DECLARE_CYCLE_STAT(TEXT("Chaos.Deformable.UDeformableSolverComponent.RemoveDeformableProxy"), STAT_ChaosDeformable_UDeformableSolverComponent_RemoveDeformableProxy, STATGROUP_Chaos);
 DECLARE_CYCLE_STAT(TEXT("Chaos.Deformable.UDeformableSolverComponent.Simulate"), STAT_ChaosDeformable_UDeformableSolverComponent_Simulate, STATGROUP_Chaos);
 DECLARE_CYCLE_STAT(TEXT("Chaos.Deformable.UDeformableSolverComponent.UpdateFromGameThread"), STAT_ChaosDeformable_UDeformableSolverComponent_UpdateFromGameThread, STATGROUP_Chaos);
 DECLARE_CYCLE_STAT(TEXT("Chaos.Deformable.UDeformableSolverComponent.UpdateFromSimulation"), STAT_ChaosDeformable_UDeformableSolverComponent_UpdateFromSimulation, STATGROUP_Chaos);
@@ -162,22 +163,24 @@ void UDeformableSolverComponent::TickComponent(float DeltaTime, enum ELevelTick
 		UpdateTickGroup();
 
 		UpdateDeformableEndTickState(IsSimulatable());
-
-		UpdateFromGameThread(DeltaTime);
-
-		if (SolverTiming.bDoThreadedAdvance)
-		{
-			// see FParallelClothCompletionTask
-			FGraphEventArray Prerequisites;
-			Prerequisites.Add(ParallelDeformableTask);
-			FGraphEventRef DeformableCompletionEvent = TGraphTask<FParallelDeformableTask>::CreateTask(&Prerequisites, ENamedThreads::GameThread).ConstructAndDispatchWhenReady(this, DeltaTime);
-			ThisTickFunction->GetCompletionHandle()->DontCompleteUntil(DeformableCompletionEvent);
-		}
-		else
+		if(bSimulationTicking)
 		{
-			Simulate(DeltaTime);
+			UpdateFromGameThread(DeltaTime);
 
-			UpdateFromSimulation(DeltaTime);
+			if (SolverTiming.bDoThreadedAdvance)
+			{
+				// see FParallelClothCompletionTask
+				FGraphEventArray Prerequisites;
+				Prerequisites.Add(ParallelDeformableTask);
+				FGraphEventRef DeformableCompletionEvent = TGraphTask<FParallelDeformableTask>::CreateTask(&Prerequisites, ENamedThreads::GameThread).ConstructAndDispatchWhenReady(this, DeltaTime);
+				ThisTickFunction->GetCompletionHandle()->DontCompleteUntil(DeformableCompletionEvent);
+			}
+			else
+			{
+				Simulate(DeltaTime);
+
+				UpdateFromSimulation(DeltaTime);
+			}
 		}
 	}
 }
@@ -239,6 +242,21 @@ void UDeformableSolverComponent::Reset()
 	}
 }
 
+void UDeformableSolverComponent::RemoveDeformableProxy(UDeformablePhysicsComponent* InComponent)
+{
+	SCOPE_CYCLE_COUNTER(STAT_ChaosDeformable_UDeformableSolverComponent_RemoveDeformableProxy);
+	TRACE_CPUPROFILER_EVENT_SCOPE(ChaosDeformable_UDeformableSolverComponent_RemoveDeformableProxy);
+
+	if (Solver && IsSimulating(InComponent))
+	{
+		FDeformableSolver::FGameThreadAccess GameThreadSolver(Solver.Get(), Chaos::Softs::FGameThreadAccessor());
+		if (!GameThreadSolver.HasObject(InComponent))
+		{
+			InComponent->RemoveProxy(GameThreadSolver);
+		}
+	}
+}
+
 void UDeformableSolverComponent::AddDeformableProxy(UDeformablePhysicsComponent* InComponent)
 {
 	SCOPE_CYCLE_COUNTER(STAT_ChaosDeformable_UDeformableSolverComponent_AddDeformableProxy);
diff --git a/Engine/Plugins/Experimental/ChaosFlesh/Source/ChaosFleshEngine/Private/ChaosFlesh/ChaosDeformableTetrahedralComponent.cpp b/Engine/Plugins/Experimental/ChaosFlesh/Source/ChaosFleshEngine/Private/ChaosFlesh/ChaosDeformableTetrahedralComponent.cpp
--- a/Engine/Plugins/Experimental/ChaosFlesh/Source/ChaosFleshEngine/Private/ChaosFlesh/ChaosDeformableTetrahedralComponent.cpp
+++ b/Engine/Plugins/Experimental/ChaosFlesh/Source/ChaosFleshEngine/Private/ChaosFlesh/ChaosDeformableTetrahedralComponent.cpp
@@ -412,7 +412,7 @@ void UDeformableTetrahedralComponent::BeginPlay()
 	}
 	if (PrimarySolverComponent)
 	{
-		PrimaryComponentTick.AddPrerequisite(PrimarySolverComponent, PrimarySolverComponent->PrimaryComponentTick);
+		PrimaryComponentTick.AddPrerequisite(PrimarySolverComponent.Get(), PrimarySolverComponent->PrimaryComponentTick);
 	}
 }
 
@@ -524,7 +524,7 @@ void UDeformableTetrahedralComponent::RenderProceduralMesh()
 								if (GetDynamicCollection())
 								{
 									const TManagedArray<FVector3f>& DynamicVertex = GetDynamicCollection()->GetPositions();
-									if (DynamicVertex.Num()) RenderVertex = DynamicVertex.GetConstArray();
+									if ((DynamicVertex.Num() > 0) && (DynamicVertex.Num() == RenderVertex.Num())) RenderVertex = DynamicVertex.GetConstArray();
 								}
 								auto InRange = [](int32 Size, int32 Val) { return 0 <= Val && Val < Size; };
 
diff --git a/Engine/Plugins/Experimental/ChaosFlesh/Source/ChaosFleshEngine/Public/ChaosFlesh/ChaosDeformablePhysicsComponent.h b/Engine/Plugins/Experimental/ChaosFlesh/Source/ChaosFleshEngine/Public/ChaosFlesh/ChaosDeformablePhysicsComponent.h
--- a/Engine/Plugins/Experimental/ChaosFlesh/Source/ChaosFleshEngine/Public/ChaosFlesh/ChaosDeformablePhysicsComponent.h
+++ b/Engine/Plugins/Experimental/ChaosFlesh/Source/ChaosFleshEngine/Public/ChaosFlesh/ChaosDeformablePhysicsComponent.h
@@ -32,6 +32,9 @@ public:
 	UFUNCTION(BlueprintCallable, Category = "Physics")
 	void EnableSimulation(UDeformableSolverComponent* DeformableSolverComponent);
 
+	UFUNCTION(BlueprintCallable, Category = "Physics")
+	void DisableSimulation();
+
 	UFUNCTION(BlueprintCallable, Category = "Physics")
 	void EnableSimulationFromActor(ADeformableSolverActor* DeformableSolverActor);
 
@@ -52,11 +55,13 @@ public:
 	UDeformableSolverComponent* GetDeformableSolver();
 	const UDeformableSolverComponent* GetDeformableSolver() const;
 	
-
+	/** Primary solver component in which this flesh component has been registered. 
+	 * Soft object ptr since the UDeformablePhysicsComponent is stored in the chaos cache which requires
+	 * soft ref in order to record itself from another thread
+	 */
 	UPROPERTY(EditAnywhere, BlueprintReadWrite, Category = "Physics", meta = (DisplayPriority = 2))
-	TObjectPtr<UDeformableSolverComponent> PrimarySolverComponent;
-
-
+	TSoftObjectPtr<UDeformableSolverComponent> PrimarySolverComponent;
+	
 	const FThreadingProxy* GetPhysicsProxy() const { return PhysicsProxy; }
 	FThreadingProxy* GetPhysicsProxy() { return PhysicsProxy; }
 
diff --git a/Engine/Plugins/Experimental/ChaosFlesh/Source/ChaosFleshEngine/Public/ChaosFlesh/ChaosDeformableSolverComponent.h b/Engine/Plugins/Experimental/ChaosFlesh/Source/ChaosFleshEngine/Public/ChaosFlesh/ChaosDeformableSolverComponent.h
--- a/Engine/Plugins/Experimental/ChaosFlesh/Source/ChaosFleshEngine/Public/ChaosFlesh/ChaosDeformableSolverComponent.h
+++ b/Engine/Plugins/Experimental/ChaosFlesh/Source/ChaosFleshEngine/Public/ChaosFlesh/ChaosDeformableSolverComponent.h
@@ -268,9 +268,11 @@ public:
 	bool IsSimulatable() const;
 	void Reset();
 	void AddDeformableProxy(UDeformablePhysicsComponent* InComponent);
+	void RemoveDeformableProxy(UDeformablePhysicsComponent* InComponent);
 	void Simulate(float DeltaTime);
 	void UpdateFromGameThread(float DeltaTime);
 	void UpdateFromSimulation(float DeltaTime);
+	void SetSimulationTicking(const bool InSimulationTicking) {bSimulationTicking = InSimulationTicking;}
 
 	/* Component Thread Management */
 	virtual void BeginPlay() override;
@@ -318,5 +320,8 @@ protected:
 	/** Ref for the deformable solvers parallel task, so we can detect whether or not a sim is running */
 	FGraphEventRef ParallelDeformableTask;
 	FDeformableEndTickFunction DeformableEndTickFunction;
+
+	/** Boolean to check if we can tick the simulation */
+	bool bSimulationTicking = true;
 };
 
diff --git a/Engine/Plugins/Experimental/ChaosFlesh/Source/ChaosFleshNodes/ChaosFleshNodes.Build.cs b/Engine/Plugins/Experimental/ChaosFlesh/Source/ChaosFleshNodes/ChaosFleshNodes.Build.cs
--- a/Engine/Plugins/Experimental/ChaosFlesh/Source/ChaosFleshNodes/ChaosFleshNodes.Build.cs
+++ b/Engine/Plugins/Experimental/ChaosFlesh/Source/ChaosFleshNodes/ChaosFleshNodes.Build.cs
@@ -5,8 +5,9 @@ namespace UnrealBuildTool.Rules
 	public class ChaosFleshNodes : ModuleRules
 	{
         public ChaosFleshNodes(ReadOnlyTargetRules Target) : base(Target)
-		{
-			PublicDependencyModuleNames.AddRange(
+        {
+	        PrivateDependencyModuleNames.AddRange(new string[] {"ChaosCaching"});
+	        PublicDependencyModuleNames.AddRange(
 				new string[]
 				{
 					"Core",
@@ -23,9 +24,10 @@ namespace UnrealBuildTool.Rules
 					"MeshConversion",
 					"MeshConversionEngineTypes",
 					"MeshDescription",
-					"TetMeshing"
+					"TetMeshing", 
+					"DataflowNodes"
 				}
 			);
-		}
+        }
 	}
 }
diff --git a/Engine/Plugins/Experimental/ChaosFlesh/Source/ChaosFleshNodes/Private/Dataflow/ChaosFleshNodesPlugin.cpp b/Engine/Plugins/Experimental/ChaosFlesh/Source/ChaosFleshNodes/Private/Dataflow/ChaosFleshNodesPlugin.cpp
--- a/Engine/Plugins/Experimental/ChaosFlesh/Source/ChaosFleshNodes/Private/Dataflow/ChaosFleshNodesPlugin.cpp
+++ b/Engine/Plugins/Experimental/ChaosFlesh/Source/ChaosFleshNodes/Private/Dataflow/ChaosFleshNodesPlugin.cpp
@@ -33,10 +33,17 @@
 #include "Dataflow/ChaosFleshVertexConstraintNode.h"
 #include "Dataflow/ChaosFleshVisualizeFiberFieldNode.h"
 #include "Dataflow/GeometryCollectionAppendCollectionTransformNode.h"
+#include "Dataflow/ChaosFleshSimulationNodes.h"
 #include "Modules/ModuleManager.h"
 
 #define LOCTEXT_NAMESPACE "ChaosFleshNodes"
 
+namespace Dataflow::CVars
+{
+	static bool bRegisterFleshSimulationNodes = true;
+	static FAutoConsoleVariableRef CVarRegisterFleshSimulationNodes(TEXT("p.Dataflow.Flesh.RegisterSimulationNodes"), bRegisterFleshSimulationNodes, TEXT("True to register all the flesh simulation nodes."));
+}
+
 
 void IChaosFleshNodesPlugin::StartupModule()
 {
@@ -72,6 +79,15 @@ void IChaosFleshNodesPlugin::StartupModule()
 	DATAFLOW_NODE_REGISTER_CREATION_FACTORY(FTriangleMeshSimulationPropertiesDataflowNodes);
 	DATAFLOW_NODE_REGISTER_CREATION_FACTORY(FVisualizeFiberFieldNode);
 	DATAFLOW_NODE_REGISTER_CREATION_FACTORY(FVisualizePositionTargetsNode);
+
+	// Temporary registry of the simulation nodes for
+	// testing while the simulation graph is being added
+	if(Dataflow::CVars::bRegisterFleshSimulationNodes)
+	{
+		Dataflow::RegisterChaosFleshSimulationNodes();
+		Dataflow::RegisterChaosSkeletonSimulationNodes();
+		Dataflow::RegisterChaosCommonSimulationNodes();
+	}
 }
 
 void IChaosFleshNodesPlugin::ShutdownModule()
diff --git a/Engine/Plugins/Experimental/ChaosFlesh/Source/ChaosFleshNodes/Private/Dataflow/ChaosFleshSimulationNodes.cpp b/Engine/Plugins/Experimental/ChaosFlesh/Source/ChaosFleshNodes/Private/Dataflow/ChaosFleshSimulationNodes.cpp
new file mode 100644
--- /dev/null
+++ b/Engine/Plugins/Experimental/ChaosFlesh/Source/ChaosFleshNodes/Private/Dataflow/ChaosFleshSimulationNodes.cpp
@@ -0,0 +1,379 @@
+// Copyright Epic Games, Inc. All Rights Reserved.
+
+#include "Dataflow/ChaosFleshSimulationNodes.h"
+#include "Dataflow/DataflowObjectInterface.h"
+#include "Dataflow/DataflowNodeFactory.h"
+#include "Dataflow/DataflowObject.h"
+#include "Chaos/Matrix.h"
+#include "Chaos/Adapters/CacheAdapter.h"
+#include "Chaos/CacheManagerActor.h"
+#include "Engine/SkeletalMesh.h"
+#include "Animation/AnimSingleNodeInstance.h"
+#include "Animation/AnimSequence.h"
+#include "Animation/AnimationPoseData.h"
+#include "BonePose.h"
+#include "BoneContainer.h"
+#include "BoneIndices.h"
+#include "Features/IModularFeatures.h"
+
+#include UE_INLINE_GENERATED_CPP_BY_NAME(ChaosFleshSimulationNodes)
+
+namespace Dataflow
+{
+	void RegisterChaosCommonSimulationNodes()
+	{
+		DATAFLOW_NODE_REGISTER_CREATION_FACTORY(FGetSimulationTimeDataflowNode);
+		
+		static const FLinearColor CDefaultNodeBodyTintColor = FLinearColor(0.f, 0.f, 0.0f, 0.5f);
+		
+		DATAFLOW_NODE_REGISTER_CREATION_FACTORY_NODE_COLORS_BY_CATEGORY("Simulation|Terminal", FLinearColor(1.0f, 0.0f, 0.0f), CDefaultNodeBodyTintColor);
+		DATAFLOW_NODE_REGISTER_CREATION_FACTORY_NODE_COLORS_BY_CATEGORY("Simulation|Setup", FLinearColor(1.0f, 1.0f, 0.0f), CDefaultNodeBodyTintColor);
+		DATAFLOW_NODE_REGISTER_CREATION_FACTORY_NODE_COLORS_BY_CATEGORY("Simulation|Update", FLinearColor(0.0f, 1.0f, 0.0f), CDefaultNodeBodyTintColor);
+	}
+	void RegisterChaosSkeletonSimulationNodes()
+	{
+		DATAFLOW_NODE_REGISTER_CREATION_FACTORY(FGetAnimationTimeRangeDataflowNode);
+		DATAFLOW_NODE_REGISTER_CREATION_FACTORY(FUpdateSkeletonAnimationDataflowNode);
+		DATAFLOW_NODE_REGISTER_CREATION_FACTORY(FSetSkeletonAnimationDataflowNode);
+	}
+	void RegisterChaosFleshSimulationNodes()
+	{
+		DATAFLOW_NODE_REGISTER_CREATION_FACTORY(FFleshSolverTerminalDataflowNode);
+		DATAFLOW_NODE_REGISTER_CREATION_FACTORY(FCreateFleshComponentDataflowNode);
+		DATAFLOW_NODE_REGISTER_CREATION_FACTORY(FCreateFleshSolverDataflowNode);
+		DATAFLOW_NODE_REGISTER_CREATION_FACTORY(FCreateFleshSkeletonDataflowNode);
+		DATAFLOW_NODE_REGISTER_CREATION_FACTORY(FAddFleshToSolverDataflowNode);
+		DATAFLOW_NODE_REGISTER_CREATION_FACTORY(FAdvanceFleshSolverDataflowNode);
+		DATAFLOW_NODE_REGISTER_CREATION_FACTORY(FBindFleshToSkeletonDataflowNode);
+	}
+}
+
+//
+// USkeletalGeneratorComponent
+//
+
+UDataflowSkeletalMeshComponent::UDataflowSkeletalMeshComponent(const FObjectInitializer& ObjectInitializer)
+	: Super(ObjectInitializer)
+{
+}
+
+UDataflowSkeletalMeshComponent::~UDataflowSkeletalMeshComponent() = default;
+
+void UDataflowSkeletalMeshComponent::FlipSpaceBuffer()
+{
+	bNeedToFlipSpaceBaseBuffers = true;
+}
+
+void FGetSimulationTimeDataflowNode::EvaluateSimulation(Dataflow::FSimulationContext& SimulationContext, const FDataflowOutput* Output) const
+{
+	SetValue(SimulationContext, FDataflowSimulationTime(SimulationContext.GetDeltaTime(), SimulationContext.GetSimulationTime()), &SimulationTime);
+}
+
+void FCreateFleshSkeletonDataflowNode::EvaluateSimulation(Dataflow::FSimulationContext& SimulationContext, const FDataflowOutput* Output) const
+{
+	const TObjectPtr<UFleshComponent> FleshPhysicsComponent = GetValue<TObjectPtr<UFleshComponent>>(SimulationContext, &FleshComponent);
+
+	if (Output->IsA<TObjectPtr<USkeletalMeshComponent>>(&SkeletalMesh))
+	{
+		if(SimulationContext.HasSimulationFlag(Dataflow::ESimulationFlags::SetupSimulation))
+		{
+			if(const TObjectPtr<AActor> RootActor = SimulationContext.GetRootActor())
+			{
+				TObjectPtr<USkeletalMeshComponent> SkeletalMeshComponent = NewObject<UDataflowSkeletalMeshComponent>(RootActor);
+				
+				SkeletalMeshComponent->SetDisablePostProcessBlueprint(true);
+				if(FleshPhysicsComponent)
+				{
+					SkeletalMeshComponent->SetSkeletalMesh(FleshPhysicsComponent->GetRestCollection()->SkeletalMesh);
+					SkeletalMeshComponent->UpdateBounds();
+				}
+
+				SkeletalMeshComponent->RegisterComponentWithWorld(RootActor->GetWorld());
+				SetValue(SimulationContext, SkeletalMeshComponent, &SkeletalMesh);
+			}
+		}
+	}
+}
+
+void FSetSkeletonAnimationDataflowNode::EvaluateSimulation(Dataflow::FSimulationContext& SimulationContext, const FDataflowOutput* Output) const
+{
+	const TObjectPtr<USkeletalMeshComponent> SkeletalMeshComponent = GetValue<TObjectPtr<USkeletalMeshComponent>>(SimulationContext, &SkeletalMesh);
+	const TObjectPtr<UAnimationAsset> FleshAnimationAsset = GetValue<TObjectPtr<UAnimationAsset>>(SimulationContext, &AnimationAsset);
+
+	if (Output->IsA<TObjectPtr<USkeletalMeshComponent>>(&SkeletalMesh))
+	{
+		if(SimulationContext.HasSimulationFlag(Dataflow::ESimulationFlags::SetupSimulation))
+		{
+			SkeletalMeshComponent->Stop();
+			SkeletalMeshComponent->AnimationData = FSingleAnimationPlayData();
+			SkeletalMeshComponent->AnimScriptInstance = nullptr;
+			
+			if(FleshAnimationAsset)
+			{
+				const TObjectPtr<UAnimSingleNodeInstance> AnimationNodeInstance = NewObject<UAnimSingleNodeInstance>(SkeletalMeshComponent);
+				AnimationNodeInstance->SetAnimationAsset(FleshAnimationAsset);
+				AnimationNodeInstance->SetPlaying(false);
+				
+				SkeletalMeshComponent->SetAnimationMode(EAnimationMode::AnimationSingleNode);
+ 				SkeletalMeshComponent->InitAnim(true);
+				SkeletalMeshComponent->AnimationData.PopulateFrom(AnimationNodeInstance);
+				SkeletalMeshComponent->AnimScriptInstance = AnimationNodeInstance;
+				SkeletalMeshComponent->AnimScriptInstance->InitializeAnimation();
+
+#if WITH_EDITOR
+				SkeletalMeshComponent->ValidateAnimation();
+#endif
+			}
+		}
+	}
+	SetValue(SimulationContext, SkeletalMeshComponent, &SkeletalMesh);
+}
+
+void FBindFleshToSkeletonDataflowNode::EvaluateSimulation(Dataflow::FSimulationContext& SimulationContext, const FDataflowOutput* Output) const
+{
+	const TObjectPtr<USkeletalMeshComponent> SkeletalMeshComponent = GetValue<TObjectPtr<USkeletalMeshComponent>>(SimulationContext, &SkeletalMesh);
+	const TObjectPtr<UDeformableSolverComponent> FleshSolverComponent = GetValue<TObjectPtr<UDeformableSolverComponent>>(SimulationContext, &FleshSolver);
+
+	if(SkeletalMeshComponent && FleshSolverComponent)
+	{
+		if(SimulationContext.HasSimulationFlag(Dataflow::ESimulationFlags::UpdateSimulation))
+		{
+			// @todo: loop over all the flesh components to find the matching one and create deformable data
+			// For now everything is done automatically in the simulate but we will need to split that
+		}
+	}
+	SetValue(SimulationContext, FleshSolverComponent, &FleshSolver);
+}
+
+void FGetAnimationTimeRangeDataflowNode::EvaluateSimulation(Dataflow::FSimulationContext& SimulationContext, const FDataflowOutput* Output) const
+{
+	const TObjectPtr<USkeletalMeshComponent> SkeletalMeshComponent = GetValue<TObjectPtr<USkeletalMeshComponent>>(SimulationContext, &SkeletalMesh);
+
+	if (Output->IsA<FVector2f>(&TimeRange))
+	{
+		if(SimulationContext.HasSimulationFlag(Dataflow::ESimulationFlags::SetupSimulation))
+		{
+			bool bHasValue = false;
+			if(SkeletalMeshComponent && SkeletalMeshComponent->GetAnimInstance())
+			{
+				if(const UAnimSingleNodeInstance* SingleNodeInstance = Cast<UAnimSingleNodeInstance>(SkeletalMeshComponent->GetAnimInstance()))
+				{
+					if(const UAnimSequenceBase* AnimSequence = Cast<UAnimSequenceBase>(SingleNodeInstance->GetAnimationAsset()))
+					{
+						const FVector2f AnimationAssetRange = FVector2f(0.0f, AnimSequence->GetPlayLength());
+						SetValue(SimulationContext, AnimationAssetRange, &TimeRange);
+						bHasValue = true;
+					}
+				}
+			}
+			if(!bHasValue)
+			{
+				SetValue(SimulationContext, FVector2f(0.0f, 5.0f), &TimeRange);
+			}
+		}
+	}
+}
+
+static void FillComponentTransforms(const int32 NumBones, const FBoneContainer& BoneContainer, const FAnimationPoseData& AnimationPoseData,
+	const FReferenceSkeleton* ReferenceSkeleton, TArray<FTransform>& ComponentSpaceTransforms) 
+{
+	ComponentSpaceTransforms.SetNumUninitialized(NumBones);
+	for (int32 Index = 0; Index < NumBones; ++Index)
+	{
+		const FCompactPoseBoneIndex CompactIndex = BoneContainer.MakeCompactPoseIndex(FMeshPoseBoneIndex(Index));
+		const int32 ParentIndex = ReferenceSkeleton->GetParentIndex(Index);
+		ComponentSpaceTransforms[Index] = (ComponentSpaceTransforms.IsValidIndex(ParentIndex) && ParentIndex < Index) ?
+			AnimationPoseData.GetPose()[CompactIndex] * ComponentSpaceTransforms[ParentIndex] : ReferenceSkeleton->GetRefBonePose()[Index];
+
+	}
+}
+
+static void FillAnimationDatas(const UAnimSequenceBase* AnimSequence, const float CurrentTime, USkeletalMesh* InSkeletalMesh, TArray<FTransform>& ComponentSpaceTransforms)
+{
+	const FAnimExtractContext ExtractionContext(FMath::Clamp(CurrentTime, 0., AnimSequence->GetPlayLength()));
+
+	const FReferenceSkeleton* ReferenceSkeleton = &InSkeletalMesh->GetRefSkeleton();
+	const int32 NumBones = ReferenceSkeleton ? ReferenceSkeleton->GetNum() : 0;
+
+	TArray<FBoneIndexType> BoneIndices;
+	BoneIndices.SetNumUninitialized(NumBones);
+	for (int32 Index = 0; Index < NumBones; ++Index)
+	{
+		BoneIndices[Index] = static_cast<FBoneIndexType>(Index);
+	}
+
+	FBoneContainer BoneContainer;
+	BoneContainer.SetUseRAWData(true);
+	BoneContainer.InitializeTo(BoneIndices, UE::Anim::FCurveFilterSettings(), *InSkeletalMesh->GetSkeleton());
+	
+	FCompactPose CompactPose;
+	CompactPose.SetBoneContainer(&BoneContainer);
+
+	FBlendedCurve BlendedCurve;
+	BlendedCurve.InitFrom(BoneContainer);
+
+	UE::Anim::FStackAttributeContainer TempAttributes;
+	FAnimationPoseData AnimationPoseData(CompactPose, BlendedCurve, TempAttributes);
+	AnimSequence->GetAnimationPose(AnimationPoseData, ExtractionContext);
+
+	FillComponentTransforms(NumBones, BoneContainer, AnimationPoseData, ReferenceSkeleton, ComponentSpaceTransforms);
+}
+
+void FUpdateSkeletonAnimationDataflowNode::EvaluateSimulation(Dataflow::FSimulationContext& SimulationContext, const FDataflowOutput* Output) const
+{
+	const TObjectPtr<USkeletalMeshComponent> SkeletalMeshComponent = GetValue<TObjectPtr<USkeletalMeshComponent>>(SimulationContext, &SkeletalMesh);
+	const float AnimationCurrentTime = GetValue<FDataflowSimulationTime>(SimulationContext, &SimulationTime).CurrentTime;
+
+	if (Output->IsA<TObjectPtr<USkeletalMeshComponent>>(&SkeletalMesh))
+	{
+		if(SimulationContext.HasSimulationFlag(Dataflow::ESimulationFlags::UpdateSimulation))
+		{
+			if(SkeletalMeshComponent && SkeletalMeshComponent->GetAnimInstance())
+			{
+				if(UAnimSingleNodeInstance* SingleNodeInstance = Cast<UAnimSingleNodeInstance>(SkeletalMeshComponent->GetAnimInstance()))
+				{
+					if(const UAnimSequenceBase* AnimSequence = Cast<UAnimSequenceBase>(SingleNodeInstance->GetAnimationAsset()))
+					{
+						TArray<FTransform> ComponentSpaceTransforms;
+						FillAnimationDatas(AnimSequence, AnimationCurrentTime, SkeletalMeshComponent->GetSkeletalMeshAsset(), ComponentSpaceTransforms);
+						
+						SkeletalMeshComponent->GetEditableComponentSpaceTransforms() = ComponentSpaceTransforms;
+						if(const TObjectPtr<UDataflowSkeletalMeshComponent> DataflowSkeletalMesh =
+							Cast<UDataflowSkeletalMeshComponent>(SkeletalMeshComponent))
+						{
+							DataflowSkeletalMesh->FlipSpaceBuffer();
+						}
+						SkeletalMeshComponent->FinalizeBoneTransform();
+						SkeletalMeshComponent->UpdateBounds();
+					}
+				}
+			}
+		}
+		SetValue(SimulationContext, SkeletalMeshComponent, &SkeletalMesh);
+	}
+}
+
+void FUpdateSkeletonAnimationDataflowNode::SetAnimationTime(Dataflow::FSimulationContext& SimulationContext, const float AnimationTime)
+{
+	const TObjectPtr<USkeletalMeshComponent> SkeletalMeshComponent = GetValue<TObjectPtr<USkeletalMeshComponent>>(SimulationContext, &SkeletalMesh);
+	
+	if(SkeletalMeshComponent && SkeletalMeshComponent->GetAnimInstance())
+	{
+		if(UAnimSingleNodeInstance* SingleNodeInstance = Cast<UAnimSingleNodeInstance>(SkeletalMeshComponent->GetAnimInstance()))
+		{
+			SingleNodeInstance->SetPosition(AnimationTime);
+		}
+	}
+}
+
+void FFleshSolverTerminalDataflowNode::EvaluateSimulation(Dataflow::FSimulationContext& SimulationContext, const FDataflowOutput* Output) const
+{
+	if(const TObjectPtr<AChaosCacheManager> ChaosCacheManager = Cast<AChaosCacheManager>(SimulationContext.GetRootActor()))
+	{
+		CacheAsset = ChaosCacheManager->CacheCollection;
+	}
+	const TObjectPtr<UDeformableSolverComponent> FleshSolverComponent = GetValue<TObjectPtr<UDeformableSolverComponent>>(SimulationContext, &FleshSolver);
+	// @todo : call the cache events from the solver interface here
+}
+
+void FCreateFleshSolverDataflowNode::EvaluateSimulation(Dataflow::FSimulationContext& SimulationContext, const FDataflowOutput* Output) const
+{
+	if(SimulationContext.HasSimulationFlag(Dataflow::ESimulationFlags::SetupSimulation))
+	{
+		if(const TObjectPtr<AActor> RootActor = SimulationContext.GetRootActor())
+		{
+			TObjectPtr<UDeformableSolverComponent> InternalFleshSolver = NewObject<UDeformableSolverComponent>(RootActor);
+		
+			InternalFleshSolver->RegisterComponentWithWorld(RootActor->GetWorld());
+		
+			InternalFleshSolver->UpdateBounds();
+			InternalFleshSolver->SetSimulationTicking(false);
+			
+			InternalFleshSolver->SolverTiming = SolverTiming;
+			InternalFleshSolver->SolverEvolution = SolverEvolution;
+			InternalFleshSolver->SolverCollisions = SolverCollisions;
+			InternalFleshSolver->SolverConstraints = SolverConstraints;
+			InternalFleshSolver->SolverForces = SolverForces;
+			InternalFleshSolver->SolverDebugging = SolverDebugging;
+			InternalFleshSolver->SolverMuscleActivation = SolverMuscleActivation;
+			InternalFleshSolver->SolverTiming.FixTimeStep = true;
+
+			SetValue(SimulationContext, InternalFleshSolver, &FleshSolver);
+		}
+	}
+}
+
+void FCreateFleshComponentDataflowNode::EvaluateSimulation(Dataflow::FSimulationContext& SimulationContext, const FDataflowOutput* Output) const
+{
+	if(SimulationContext.HasSimulationFlag(Dataflow::ESimulationFlags::SetupSimulation))
+	{
+		if(const TObjectPtr<AActor> RootActor = SimulationContext.GetRootActor())
+		{	
+			const FName ChannelName(SimulationContext.Graph->GetName() + TEXT("_") + GetName().ToString() + TEXT("_FleshComponent"));
+			
+			TObjectPtr<UFleshComponent> InternalFleshComponent = NewObject<UFleshComponent>(RootActor, ChannelName);
+			InternalFleshComponent->RegisterComponentWithWorld(RootActor->GetWorld());
+
+			// Create and add the visualisation component
+			if(InternalFleshComponent->Mesh)
+			{
+				InternalFleshComponent->Mesh->RegisterComponentWithWorld(RootActor->GetWorld());
+			}
+			InternalFleshComponent->SetRestCollection(FleshAsset);
+			InternalFleshComponent->UpdateBounds();
+
+			if(const TObjectPtr<AChaosCacheManager> ChaosCacheManager = Cast<AChaosCacheManager>(SimulationContext.GetRootActor()))
+			{
+				// Get the implementation of our adapters for identifying compatible components
+				IModularFeatures&                      ModularFeatures = IModularFeatures::Get();
+				TArray<Chaos::FComponentCacheAdapter*> Adapters = ModularFeatures.GetModularFeatureImplementations<Chaos::FComponentCacheAdapter>(Chaos::FComponentCacheAdapter::FeatureName);
+		
+				if(Chaos::FAdapterUtil::GetBestAdapterForClass(InternalFleshComponent->GetClass(), false))
+				{
+					ChaosCacheManager->FindOrAddObservedComponent(InternalFleshComponent, ChannelName, true);
+				}
+			}
+			SetValue(SimulationContext, InternalFleshComponent, &FleshComponent);
+		}
+	}
+}
+
+void FAddFleshToSolverDataflowNode::EvaluateSimulation(Dataflow::FSimulationContext& SimulationContext, const FDataflowOutput* Output) const
+{
+	const TObjectPtr<UDeformableSolverComponent> FleshSolverComponent = GetValue<TObjectPtr<UDeformableSolverComponent>>(SimulationContext, &FleshSolver);
+	const TObjectPtr<UFleshComponent> FleshPhysicsComponent = GetValue<TObjectPtr<UFleshComponent>>(SimulationContext, &FleshComponent);
+
+	if(FleshSolverComponent && FleshPhysicsComponent)
+	{
+		if(SimulationContext.HasSimulationFlag(Dataflow::ESimulationFlags::SetupSimulation))
+		{
+			FleshPhysicsComponent->EnableSimulation(FleshSolverComponent);
+			FleshSolverComponent->Reset();
+		}
+	}
+	
+	SetValue(SimulationContext, FleshSolverComponent, &FleshSolver);
+}
+
+void FAdvanceFleshSolverDataflowNode::EvaluateSimulation(Dataflow::FSimulationContext& SimulationContext, const FDataflowOutput* Output) const
+{
+	const TObjectPtr<UDeformableSolverComponent> FleshSolverComponent = GetValue<TObjectPtr<UDeformableSolverComponent>>(SimulationContext, &FleshSolver);
+	const float SimulationDeltaTime = GetValue<FDataflowSimulationTime>(SimulationContext, &SimulationTime).DeltaTime;
+	
+	if(FleshSolverComponent && (SimulationContext.GetCachingMode() == Dataflow::ECachingMode::RecordCache) &&
+		SimulationContext.HasSimulationFlag(Dataflow::ESimulationFlags::UpdateSimulation))
+	{
+		const float SimulationDuration = SimulationDeltaTime + TimeOffset;
+		const int32 NumSteps = FMath::Floor(SimulationDuration / FleshSolverComponent->SolverTiming.TimeStepSize);
+
+		for(int32 StepIndex = 0; StepIndex < NumSteps; ++StepIndex)
+		{
+			FleshSolverComponent->UpdateFromGameThread(FleshSolverComponent->SolverTiming.TimeStepSize);
+			FleshSolverComponent->Simulate(FleshSolverComponent->SolverTiming.TimeStepSize);
+			FleshSolverComponent->UpdateFromSimulation(FleshSolverComponent->SolverTiming.TimeStepSize);
+		}
+		TimeOffset = SimulationDuration - NumSteps * FleshSolverComponent->SolverTiming.TimeStepSize;
+	}
+	
+	SetValue(SimulationContext, FleshSolverComponent, &FleshSolver);
+}
\ No newline at end of file
diff --git a/Engine/Plugins/Experimental/ChaosFlesh/Source/ChaosFleshNodes/Public/Dataflow/ChaosFleshSimulationNodes.h b/Engine/Plugins/Experimental/ChaosFlesh/Source/ChaosFleshNodes/Public/Dataflow/ChaosFleshSimulationNodes.h
new file mode 100644
--- /dev/null
+++ b/Engine/Plugins/Experimental/ChaosFlesh/Source/ChaosFleshNodes/Public/Dataflow/ChaosFleshSimulationNodes.h
@@ -0,0 +1,413 @@
+// Copyright Epic Games, Inc. All Rights Reserved.
+
+#pragma once 
+
+#include "CoreMinimal.h"
+#include "ChaosLog.h"
+#include "ChaosFlesh/FleshAsset.h"
+#include "ChaosFlesh/FleshComponent.h"
+#include "ChaosFlesh/ChaosDeformableSolverComponent.h"
+#include "Chaos/CacheCollection.h"
+#include "Dataflow/DataflowSolverNodes.h"
+#include "Components/SkeletalMeshComponent.h"
+
+#include "ChaosFleshSimulationNodes.generated.h"
+
+/**
+ * UDataflowSkeletalMesh component
+ */
+UCLASS()
+class CHAOSFLESHNODES_API UDataflowSkeletalMeshComponent : public USkeletalMeshComponent
+{
+	GENERATED_BODY()
+public:
+	UDataflowSkeletalMeshComponent(const FObjectInitializer& ObjectInitializer);
+	~UDataflowSkeletalMeshComponent();
+
+	/** Flip space buffer if needed */
+	void FlipSpaceBuffer();
+};
+
+/** Get the dataflow simulation time */
+USTRUCT()
+struct FDataflowSimulationTime 
+{
+	GENERATED_BODY()
+
+public:
+	FDataflowSimulationTime() : DeltaTime(0.0f), CurrentTime(0.0f)
+	{}
+	
+	FDataflowSimulationTime(const float InDeltaTime, const float InCurrentTime) : DeltaTime(InDeltaTime), CurrentTime(InCurrentTime)
+	{}
+	
+	/** Delta time in seconds property coming from the context */
+	UPROPERTY(Transient, SkipSerialization)
+	float DeltaTime = 0.0f;
+
+	/** Current time in seconds property coming from the context */
+	UPROPERTY(Transient, SkipSerialization)
+	float CurrentTime = 0.0f;
+};
+
+/** Get the context simulation time */
+USTRUCT(meta = (DataflowSimulation))
+struct FGetSimulationTimeDataflowNode : public FDataflowInvalidNode
+{
+	GENERATED_USTRUCT_BODY()
+	DATAFLOW_NODE_DEFINE_INTERNAL(FGetSimulationTimeDataflowNode, "GetSimulationTime", "Simulation|Update|Common", "")
+
+public:
+	
+	FGetSimulationTimeDataflowNode(const Dataflow::FNodeParameters& InParam, FGuid InGuid = FGuid::NewGuid())
+	: FDataflowInvalidNode(InParam, InGuid)
+	{
+		RegisterOutputConnection(&SimulationTime);
+	}
+	
+	/** Simulation time property coming from the context */
+	UPROPERTY(Transient, SkipSerialization, Meta = (DataflowOutput))
+	FDataflowSimulationTime SimulationTime = FDataflowSimulationTime(0.0f, 0.0f);
+
+	/** Evaluate simulation node given a simulation context */
+	virtual void EvaluateSimulation(Dataflow::FSimulationContext& SimulationContext, const FDataflowOutput* Output) const override;
+};
+
+/** Bind the skeletal mesh to the matching flesh component within the flesh solver */
+USTRUCT(meta = (DataflowCache))
+struct FBindFleshToSkeletonDataflowNode : public FDataflowSimulationNode
+{
+	GENERATED_USTRUCT_BODY()
+	DATAFLOW_NODE_DEFINE_INTERNAL(FBindFleshToSkeletonDataflowNode, "BindFleshToSkeleton", "Simulation|Update|Flesh", "")
+
+public:
+	
+	FBindFleshToSkeletonDataflowNode(const Dataflow::FNodeParameters& InParam, FGuid InGuid = FGuid::NewGuid())
+	: FDataflowSimulationNode(InParam, InGuid)
+	{
+		RegisterInputConnection(&SkeletalMesh);
+		RegisterInputConnection(&FleshSolver);
+		RegisterOutputConnection(&FleshSolver, &FleshSolver);
+	}
+
+	/** Flesh solver from which we will bind one flesh component to the skeleton mesh */
+	UPROPERTY(Transient, SkipSerialization, Meta = (DataflowInput, DataflowOutput, DataflowPassthrough = "FleshSolver"))
+	TObjectPtr<UDeformableSolverComponent> FleshSolver = nullptr;
+
+	/** Skeleton mesh used for binding*/
+	UPROPERTY(Transient, SkipSerialization, Meta = (DataflowInput))
+	TObjectPtr<USkeletalMeshComponent> SkeletalMesh = nullptr;
+	
+	/** Evaluate simulation node given a simulation context */
+	virtual void EvaluateSimulation(Dataflow::FSimulationContext& SimulationContext, const FDataflowOutput* Output) const override;
+};
+
+/** Main terminal node for flesh solver */
+USTRUCT(meta = (DataflowCache))
+struct FFleshSolverTerminalDataflowNode : public FDataflowCacheNode
+{
+	GENERATED_USTRUCT_BODY()
+	DATAFLOW_NODE_DEFINE_INTERNAL(FFleshSolverTerminalDataflowNode, "FleshSolverTerminal", "Simulation|Terminal|Flesh", "")
+
+public:
+	
+	FFleshSolverTerminalDataflowNode(const Dataflow::FNodeParameters& InParam, FGuid InGuid = FGuid::NewGuid())
+	: FDataflowCacheNode(InParam, InGuid)
+	{
+		RegisterInputConnection(&FleshSolver);
+		RegisterInputConnection(&TimeRange);
+	}
+
+	/** Cache asset that will store the cache datas */
+	UPROPERTY(VisibleAnywhere, Category="Caching")
+	mutable TObjectPtr<UChaosCacheCollection> CacheAsset = nullptr;
+
+	/** Time range used for caching */
+	UPROPERTY(EditAnywhere, Category="Caching", Meta = (DataflowInput))
+	FVector2f TimeRange = FVector2f(0.0f, 5.0f);
+
+	/** Flesh solver that will be cached */
+	UPROPERTY(Transient, SkipSerialization, Meta = (DataflowInput))
+	TObjectPtr<UDeformableSolverComponent> FleshSolver = nullptr;
+	
+	/** Evaluate simulation node given a simulation context */
+	virtual void EvaluateSimulation(Dataflow::FSimulationContext& SimulationContext, const FDataflowOutput* Output) const override;
+	
+	/** Get the cache asset linked to that terminal cache */
+	virtual TObjectPtr<UObject> GetCacheAsset(Dataflow::FSimulationContext& SimulationContext) const override
+	{
+		return Cast<UObject>(GetValue<TObjectPtr<UChaosCacheCollection>>(SimulationContext, &CacheAsset));
+	}
+	
+	/** Return the time range used to cache the datas */
+	virtual FVector2f GetTimeRange(Dataflow::FSimulationContext& SimulationContext) const override
+	{
+		return GetValue<FVector2f>(SimulationContext, &TimeRange);
+	}
+};
+
+/** Create flesh solver from settings */
+USTRUCT(meta = (DataflowFlesh))
+struct FCreateFleshSolverDataflowNode : public FDataflowSimulationNode
+{
+	GENERATED_USTRUCT_BODY()
+	DATAFLOW_NODE_DEFINE_INTERNAL(FCreateFleshSolverDataflowNode, "CreateFleshSolver", "Simulation|Setup|Flesh", "")
+
+public:
+	
+	FCreateFleshSolverDataflowNode(const Dataflow::FNodeParameters& InParam, FGuid InGuid = FGuid::NewGuid())
+	: FDataflowSimulationNode(InParam, InGuid)
+	{
+		RegisterOutputConnection(&FleshSolver);
+	}
+
+	/** Flesh solver built from the settings */
+	UPROPERTY(Transient, SkipSerialization, Meta = (DataflowOutput))
+	TObjectPtr<UDeformableSolverComponent> FleshSolver = nullptr;
+
+	/** Solver timing properties */
+	UPROPERTY(EditAnywhere, Category="Solver")
+	FSolverTimingGroup SolverTiming;
+
+	/** Solver evolution properties */
+	UPROPERTY(EditAnywhere, Category="Solver")
+	FSolverEvolutionGroup SolverEvolution;
+
+	/** Solver collisions properties */
+	UPROPERTY(EditAnywhere, Category="Solver")
+	FSolverCollisionsGroup SolverCollisions;
+
+	/** Solver constraints properties */
+	UPROPERTY(EditAnywhere, Category="Solver")
+	FSolverConstraintsGroup SolverConstraints;
+
+	/** Solver forces properties */
+	UPROPERTY(EditAnywhere, Category="Solver")
+	FSolverForcesGroup SolverForces;
+
+	/** Solver debugging properties */
+	UPROPERTY(EditAnywhere, Category="Solver")
+	FSolverDebuggingGroup SolverDebugging;
+
+	/** Solver muscle activation properties */
+	UPROPERTY(EditAnywhere, Category="Solver")
+	FSolverMuscleActivationGroup SolverMuscleActivation;
+
+	/** Evaluate simulation node given a simulation context */
+	virtual void EvaluateSimulation(Dataflow::FSimulationContext& SimulationContext, const FDataflowOutput* Output) const override;
+};
+
+/** Create flesh skeleton from settings */
+USTRUCT(meta = (DataflowFlesh))
+struct FCreateFleshSkeletonDataflowNode : public FDataflowSimulationNode
+{
+	GENERATED_USTRUCT_BODY()
+	DATAFLOW_NODE_DEFINE_INTERNAL(FCreateFleshSkeletonDataflowNode, "CreateFleshSkeleton", "Simulation|Setup|Flesh", "")
+
+public:
+	
+	FCreateFleshSkeletonDataflowNode(const Dataflow::FNodeParameters& InParam, FGuid InGuid = FGuid::NewGuid())
+	: FDataflowSimulationNode(InParam, InGuid)
+	{
+		RegisterInputConnection(&FleshComponent);
+		RegisterOutputConnection(&SkeletalMesh);
+	}
+
+	/** Skeletal mesh component to be built */
+	UPROPERTY(Transient, SkipSerialization, Meta = (DataflowOutput))
+	TObjectPtr<USkeletalMeshComponent> SkeletalMesh = nullptr;
+
+	/** Flesh component from which the skeletal mesh will be constructed */
+	UPROPERTY(Transient, SkipSerialization, Meta = (DataflowInput))
+	TObjectPtr<UFleshComponent> FleshComponent = nullptr;
+
+	/** Evaluate simulation node given a simulation context */
+	virtual void EvaluateSimulation(Dataflow::FSimulationContext& SimulationContext, const FDataflowOutput* Output) const override;
+};
+
+
+/** Set a specified animation onto the skeleton mesh */
+USTRUCT(meta = (DataflowFlesh))
+struct FSetSkeletonAnimationDataflowNode : public FDataflowSimulationNode
+{
+	GENERATED_USTRUCT_BODY()
+	DATAFLOW_NODE_DEFINE_INTERNAL(FSetSkeletonAnimationDataflowNode, "SetSkeletonAnimation", "Simulation|Setup|Skeleton", "")
+
+public:
+	
+	FSetSkeletonAnimationDataflowNode(const Dataflow::FNodeParameters& InParam, FGuid InGuid = FGuid::NewGuid())
+	: FDataflowSimulationNode(InParam, InGuid)
+	{
+		RegisterInputConnection(&SkeletalMesh);
+		RegisterInputConnection(&AnimationAsset);
+		RegisterOutputConnection(&SkeletalMesh, &SkeletalMesh);
+	}
+
+	/** Skeletal mesh component to be built */
+	UPROPERTY(Transient, SkipSerialization, Meta = (DataflowInput, DataflowOutput, DataflowPassthrough = "SkeletalMesh"))
+	TObjectPtr<USkeletalMeshComponent> SkeletalMesh = nullptr;
+
+	/** Animation asset to be used  */
+	UPROPERTY(EditAnywhere, Category="Animation", Meta = (DataflowInput))
+	TObjectPtr<UAnimationAsset> AnimationAsset = nullptr;
+
+	/** Evaluate simulation node given a simulation context */
+	virtual void EvaluateSimulation(Dataflow::FSimulationContext& SimulationContext, const FDataflowOutput* Output) const override;
+};
+
+/** Get the skeleton animation time range */
+USTRUCT(meta = (DataflowFlesh))
+struct FGetAnimationTimeRangeDataflowNode : public FDataflowSimulationNode
+{
+	GENERATED_USTRUCT_BODY()
+	DATAFLOW_NODE_DEFINE_INTERNAL(FGetAnimationTimeRangeDataflowNode, "GetAnimationTimeRange", "Simulation|Setup|Skeleton", "")
+
+public:
+	
+	FGetAnimationTimeRangeDataflowNode(const Dataflow::FNodeParameters& InParam, FGuid InGuid = FGuid::NewGuid())
+	: FDataflowSimulationNode(InParam, InGuid)
+	{
+		RegisterInputConnection(&SkeletalMesh);
+		RegisterOutputConnection(&TimeRange);
+	}
+
+	/** Skeletal mesh component yo update at a given time */
+	UPROPERTY(Transient, SkipSerialization, Meta = (DataflowInput))
+	TObjectPtr<USkeletalMeshComponent> SkeletalMesh = nullptr;
+
+	/** Time range in the skeletal mesh animation instance */
+	UPROPERTY(Transient, SkipSerialization,  Meta = (DataflowOutput))
+	FVector2f TimeRange = FVector2f(0.0f,5.0f);
+
+	/** Evaluate simulation node given a simulation context */
+	virtual void EvaluateSimulation(Dataflow::FSimulationContext& SimulationContext, const FDataflowOutput* Output) const override;
+};
+
+/** Update skeleton animation at a given time */
+USTRUCT(meta = (DataflowFlesh))
+struct FUpdateSkeletonAnimationDataflowNode : public FDataflowAnimationNode
+{
+	GENERATED_USTRUCT_BODY()
+	DATAFLOW_NODE_DEFINE_INTERNAL(FUpdateSkeletonAnimationDataflowNode, "UpdateSkeletonAnimation", "Simulation|Update|Skeleton", "")
+
+public:
+	
+	FUpdateSkeletonAnimationDataflowNode(const Dataflow::FNodeParameters& InParam, FGuid InGuid = FGuid::NewGuid())
+	: FDataflowAnimationNode(InParam, InGuid)
+	{
+		RegisterInputConnection(&SimulationTime);
+		RegisterInputConnection(&SkeletalMesh);
+		RegisterOutputConnection(&SkeletalMesh);
+	}
+
+	/** Skeletal mesh component yo update at a given time */
+	UPROPERTY(Transient, SkipSerialization, Meta = (DataflowInput, DataflowOutput))
+	TObjectPtr<USkeletalMeshComponent> SkeletalMesh = nullptr;
+
+	/** Simulation time used to load the animation */
+	UPROPERTY(Transient, SkipSerialization,  Meta = (DataflowInput))
+	FDataflowSimulationTime SimulationTime = FDataflowSimulationTime(0.0f,0.0f);
+
+	/** Evaluate simulation node given a simulation context */
+	virtual void EvaluateSimulation(Dataflow::FSimulationContext& SimulationContext, const FDataflowOutput* Output) const override;
+
+	/** Set the animation time */
+	virtual void SetAnimationTime(Dataflow::FSimulationContext& SimulationContext, const float AnimationTime) override;
+};
+
+/** Create a flesh component given a flesh asset */
+USTRUCT(meta = (DataflowFlesh))
+struct FCreateFleshComponentDataflowNode : public FDataflowSimulationNode
+{
+	GENERATED_USTRUCT_BODY()
+	DATAFLOW_NODE_DEFINE_INTERNAL(FCreateFleshComponentDataflowNode, "CreateFleshComponent", "Simulation|Setup|Flesh", "")
+
+public:
+	
+	FCreateFleshComponentDataflowNode(const Dataflow::FNodeParameters& InParam, FGuid InGuid = FGuid::NewGuid())
+	: FDataflowSimulationNode(InParam, InGuid)
+	{
+		RegisterInputConnection(&FleshAsset);
+		RegisterOutputConnection(&FleshComponent);
+	}
+
+	/** Flesh component built from the asset */
+	UPROPERTY(Transient, SkipSerialization, Meta = (DataflowOutput))
+	TObjectPtr<UFleshComponent> FleshComponent = nullptr;
+
+	/** Flesh asset */
+	UPROPERTY(EditAnywhere, Category="Flesh", Meta = (DataflowInput))
+	TObjectPtr<UFleshAsset> FleshAsset = nullptr;
+
+	/** Evaluate simulation node given a simulation context */
+	virtual void EvaluateSimulation(Dataflow::FSimulationContext& SimulationContext, const FDataflowOutput* Output) const override;
+};
+
+/** Add flesh component to a flesh solver */
+USTRUCT(meta = (DataflowFlesh))
+struct FAddFleshToSolverDataflowNode : public FDataflowSimulationNode
+{
+	GENERATED_USTRUCT_BODY()
+	DATAFLOW_NODE_DEFINE_INTERNAL(FAddFleshToSolverDataflowNode, "AddFleshToSolver", "Simulation|Setup|Flesh", "")
+
+public:
+	
+	FAddFleshToSolverDataflowNode(const Dataflow::FNodeParameters& InParam, FGuid InGuid = FGuid::NewGuid())
+	: FDataflowSimulationNode(InParam, InGuid)
+	{
+		RegisterInputConnection(&FleshComponent);
+		RegisterInputConnection(&FleshSolver);
+		RegisterOutputConnection(&FleshSolver, &FleshSolver);
+	}
+
+	/** Flesh component to add to the solver */
+	UPROPERTY(Transient, SkipSerialization, Meta = (DataflowInput))
+	TObjectPtr<UFleshComponent> FleshComponent = nullptr;
+
+	/** Flesh solver */
+	UPROPERTY(Transient, SkipSerialization, Meta = (DataflowInput, DataflowOutput, DataflowPassthrough = "FleshSolver"))
+	TObjectPtr<UDeformableSolverComponent> FleshSolver = nullptr;
+
+	/** Evaluate simulation node given a simulation context */
+	virtual void EvaluateSimulation(Dataflow::FSimulationContext& SimulationContext, const FDataflowOutput* Output) const override;
+};
+
+/** Advance the flesh solver in time */
+USTRUCT(meta = (DataflowFlesh))
+struct FAdvanceFleshSolverDataflowNode : public FDataflowSimulationNode
+{
+	GENERATED_USTRUCT_BODY()
+	DATAFLOW_NODE_DEFINE_INTERNAL(FAdvanceFleshSolverDataflowNode, "AdvanceFleshSolver", "Simulation|Update|Flesh", "")
+
+public:
+	
+	FAdvanceFleshSolverDataflowNode(const Dataflow::FNodeParameters& InParam, FGuid InGuid = FGuid::NewGuid())
+	: FDataflowSimulationNode(InParam, InGuid)
+	{
+		RegisterInputConnection(&SimulationTime);
+		RegisterInputConnection(&FleshSolver);
+		RegisterOutputConnection(&FleshSolver, &FleshSolver);
+	}
+
+	/** Delta time to use to advance the solver*/
+	UPROPERTY(Transient, SkipSerialization, Meta = (DataflowInput))
+	FDataflowSimulationTime SimulationTime = FDataflowSimulationTime(0.0f,0.0f);
+
+	/** Flesh solver */
+	UPROPERTY(Transient, SkipSerialization, Meta = (DataflowInput, DataflowOutput, DataflowPassthrough = "FleshSolver"))
+	TObjectPtr<UDeformableSolverComponent> FleshSolver = nullptr;
+
+	/** Evaluate simulation node given a simulation context */
+	virtual void EvaluateSimulation(Dataflow::FSimulationContext& SimulationContext, const FDataflowOutput* Output) const override;
+
+private :
+	/** Time offset to store the difference in between the delta time and the fixed time step */
+	mutable float TimeOffset = 0.0f;
+};
+
+namespace Dataflow
+{
+	void RegisterChaosCommonSimulationNodes();
+	void RegisterChaosFleshSimulationNodes();
+	void RegisterChaosSkeletonSimulationNodes();
+}
