ChaosFlesh: Visualize kinematic vertices

diff --git a/Engine/Plugins/Experimental/ChaosFlesh/Source/ChaosFleshNodes/Private/Dataflow/ChaosFleshNodesPlugin.cpp b/Engine/Plugins/Experimental/ChaosFlesh/Source/ChaosFleshNodes/Private/Dataflow/ChaosFleshNodesPlugin.cpp
--- a/Engine/Plugins/Experimental/ChaosFlesh/Source/ChaosFleshNodes/Private/Dataflow/ChaosFleshNodesPlugin.cpp
+++ b/Engine/Plugins/Experimental/ChaosFlesh/Source/ChaosFleshNodes/Private/Dataflow/ChaosFleshNodesPlugin.cpp
@@ -79,6 +79,7 @@ void IChaosFleshNodesPlugin::StartupModule()
 	DATAFLOW_NODE_REGISTER_CREATION_FACTORY(FTriangleMeshSimulationPropertiesDataflowNodes);
 	DATAFLOW_NODE_REGISTER_CREATION_FACTORY(FVisualizeFiberFieldNode);
 	DATAFLOW_NODE_REGISTER_CREATION_FACTORY(FVisualizePositionTargetsNode);
+	DATAFLOW_NODE_REGISTER_CREATION_FACTORY(FVisualizeKinematicFacesNode);
 
 	// Temporary registry of the simulation nodes for
 	// testing while the simulation graph is being added
diff --git a/Engine/Plugins/Experimental/ChaosFlesh/Source/ChaosFleshNodes/Private/Dataflow/ChaosFleshVisualizeFiberFieldNode.cpp b/Engine/Plugins/Experimental/ChaosFlesh/Source/ChaosFleshNodes/Private/Dataflow/ChaosFleshVisualizeFiberFieldNode.cpp
--- a/Engine/Plugins/Experimental/ChaosFlesh/Source/ChaosFleshNodes/Private/Dataflow/ChaosFleshVisualizeFiberFieldNode.cpp
+++ b/Engine/Plugins/Experimental/ChaosFlesh/Source/ChaosFleshNodes/Private/Dataflow/ChaosFleshVisualizeFiberFieldNode.cpp
@@ -1,8 +1,8 @@
 // Copyright Epic Games, Inc. All Rights Reserved.
 
 #include "Dataflow/ChaosFleshVisualizeFiberFieldNode.h"
+#include "GeometryCollection/Facades/CollectionKinematicBindingFacade.h"
 #include "GeometryCollection/Facades/CollectionPositionTargetFacade.h"
-
 #include "ChaosFlesh/TetrahedralCollection.h"
 
 #include UE_INLINE_GENERATED_CPP_BY_NAME(ChaosFleshVisualizeFiberFieldNode)
@@ -74,4 +74,58 @@ void FVisualizePositionTargetsNode::Evaluate(Dataflow::FContext& Context, const
 
 		Out->SetValue(MoveTemp(OutVectorField), Context);
 	}
+}
+
+void FVisualizeKinematicFacesNode::Evaluate(Dataflow::FContext& Context, const FDataflowOutput* Out) const
+{
+	if (Out->IsA<FManagedArrayCollection>(&Collection))
+	{
+		FManagedArrayCollection InCollection = GetValue<FManagedArrayCollection>(Context, &Collection);
+		TManagedArray<FIntVector>* Indices = InCollection.FindAttribute<FIntVector>("Indices", FGeometryCollection::FacesGroup);
+		TManagedArray<bool>* FaceVisibility = InCollection.FindAttribute<bool>("Visible", FGeometryCollection::FacesGroup);
+		TManagedArray<FVector3f>* Vertices = InCollection.FindAttribute<FVector3f>("Vertex", FGeometryCollection::VerticesGroup);
+		if (Indices && FaceVisibility && Vertices)
+		{
+			//Find kinematic particles
+			TArray<bool> ParticleIsKinematic;
+			ParticleIsKinematic.Init(false, Vertices->Num());
+			typedef GeometryCollection::Facades::FKinematicBindingFacade FKinematics;
+			FKinematics Kinematics(InCollection);
+			//RemoveElements
+			// Add Kinematics Node
+			for (int i = Kinematics.NumKinematicBindings() - 1; i >= 0; i--)
+			{
+				FKinematics::FBindingKey Key = Kinematics.GetKinematicBindingKey(i);
+
+				int32 BoneIndex = INDEX_NONE;
+				TArray<int32> BoundVerts;
+				TArray<float> BoundWeights;
+				Kinematics.GetBoneBindings(Key, BoneIndex, BoundVerts, BoundWeights);
+
+				for (int32 vdx : BoundVerts)
+				{
+					ParticleIsKinematic[vdx] = true;
+				}
+			}
+
+			FaceVisibility->Fill(false);
+			for (int32 FaceIdx = 0; FaceIdx < Indices->Num(); ++FaceIdx)
+			{
+				bool IsElementKinematic = true;
+				for (int32 j = 0; j < 3; j++)
+				{
+					if (!ParticleIsKinematic[(*Indices)[FaceIdx][j]])
+					{
+						IsElementKinematic = false;
+						break;
+					}
+				}
+				if (IsElementKinematic)
+				{
+					(*FaceVisibility)[FaceIdx] = true;
+				}
+			}
+		}
+		SetValue(Context, MoveTemp(InCollection), &Collection);
+	}
 }
\ No newline at end of file
diff --git a/Engine/Plugins/Experimental/ChaosFlesh/Source/ChaosFleshNodes/Public/Dataflow/ChaosFleshVisualizeFiberFieldNode.h b/Engine/Plugins/Experimental/ChaosFlesh/Source/ChaosFleshNodes/Public/Dataflow/ChaosFleshVisualizeFiberFieldNode.h
--- a/Engine/Plugins/Experimental/ChaosFlesh/Source/ChaosFleshNodes/Public/Dataflow/ChaosFleshVisualizeFiberFieldNode.h
+++ b/Engine/Plugins/Experimental/ChaosFlesh/Source/ChaosFleshNodes/Public/Dataflow/ChaosFleshVisualizeFiberFieldNode.h
@@ -63,5 +63,29 @@ public:
 		RegisterOutputConnection(&VectorField);
 	}
 
+	virtual void Evaluate(Dataflow::FContext& Context, const FDataflowOutput* Out) const override;
+};
+
+/**
+* Visualizes kinematic faces from GeometryCollection.
+*/
+USTRUCT(meta = (DataflowFlesh))
+struct FVisualizeKinematicFacesNode : public FDataflowNode
+{
+	GENERATED_USTRUCT_BODY()
+	DATAFLOW_NODE_DEFINE_INTERNAL(FVisualizeKinematicFacesNode, "VisualizeKinematicFaces", "Flesh", "")
+	DATAFLOW_NODE_RENDER_TYPE(FGeometryCollection::StaticType(), "Collection")
+
+public:
+	UPROPERTY(meta = (DataflowInput, DataflowOutput, DisplayName = "Collection", DataflowPassthrough = "Collection"))
+	FManagedArrayCollection Collection;
+
+	FVisualizeKinematicFacesNode(const Dataflow::FNodeParameters& InParam, FGuid InGuid = FGuid::NewGuid())
+		: FDataflowNode(InParam, InGuid)
+	{
+		RegisterInputConnection(&Collection);
+		RegisterOutputConnection(&Collection, &Collection);
+	}
+
 	virtual void Evaluate(Dataflow::FContext& Context, const FDataflowOutput* Out) const override;
 };
\ No newline at end of file
