ChaosFlesh: update muscle volume parameter

diff --git a/Engine/Plugins/Experimental/ChaosFlesh/Source/ChaosFleshNodes/Private/Dataflow/ChaosFleshComputeMuscleActivationNode.cpp b/Engine/Plugins/Experimental/ChaosFlesh/Source/ChaosFleshNodes/Private/Dataflow/ChaosFleshComputeMuscleActivationNode.cpp
--- a/Engine/Plugins/Experimental/ChaosFlesh/Source/ChaosFleshNodes/Private/Dataflow/ChaosFleshComputeMuscleActivationNode.cpp
+++ b/Engine/Plugins/Experimental/ChaosFlesh/Source/ChaosFleshNodes/Private/Dataflow/ChaosFleshComputeMuscleActivationNode.cpp
@@ -16,6 +16,7 @@ void FComputeMuscleActivationDataNode::Evaluate(Dataflow::FContext& Context, con
 		TArray<int32> InOriginIndices = GetValue<TArray<int32>>(Context, &OriginIndicesIn);
 		TArray<int32> InInsertionIndices = GetValue<TArray<int32>>(Context, &InsertionIndicesIn);
 		TArray<float> OriginInsertionRestLength;
+		TManagedArray<float>* ContractionVolumeScalePerVertex = InCollection.FindAttribute<float>("ContractionVolumeScale", "Vertices");
 		if (InOriginIndices.Num() > 0 && InInsertionIndices.Num() > 0)
 		{
 			if (TManagedArray<FVector3f>* Vertex = InCollection.FindAttribute<FVector3f>("Vertex", "Vertices"))
@@ -26,55 +27,71 @@ void FComputeMuscleActivationDataNode::Evaluate(Dataflow::FContext& Context, con
 					{
 						GeometryCollection::Facades::FMuscleActivationFacade FMuscleActivation(InCollection);
 						TArray<TArray<int32>> MuscleActivationElements;
-						TArray<int32> OriginNodes;
-						TArray<int32> InsertionNodes;
+						TArray<TArray<int32>> ComponentOrigins; //One origin node per muscle component
+						TArray<TArray<int32>> ComponentInsertions; //One insertion node per muscle component
 						GeometryCollection::Facades::FCollectionMeshFacade MeshFacade(InCollection);
 						TArray<int32> ComponentIndex = MeshFacade.GetGeometryGroupIndexArray();
-						TMap<int32, int32> ComponentToIndex; //Component index to muscle index
+						TMap<int32, int32> ComponentToMuscleIndex; //Component index to muscle index
 						for (int32 i = 0; i < InOriginIndices.Num(); i++)
 						{
-							if (!ComponentToIndex.Contains(ComponentIndex[InOriginIndices[i]]))
+							if (!ComponentToMuscleIndex.Contains(ComponentIndex[InOriginIndices[i]]))
 							{
-								//TODO: use one pair of separately defined origin and insertion for each muscle instead of choosing the first from kinematic origins and last from insertions
-								ComponentToIndex.Add(ComponentIndex[InOriginIndices[i]], OriginNodes.Num());
-								OriginNodes.Add(InOriginIndices[i]);
+								ComponentToMuscleIndex.Add(ComponentIndex[InOriginIndices[i]], ComponentOrigins.Num());
+								ComponentOrigins.SetNum(ComponentOrigins.Num() + 1);
+								ComponentOrigins[ComponentOrigins.Num()-1].Add(InOriginIndices[i]);
+							}
+							else
+							{
+								ComponentOrigins[ComponentToMuscleIndex[ComponentIndex[InOriginIndices[i]]]].Add(InOriginIndices[i]);
 							}
 						}
-						InsertionNodes.Init(INDEX_NONE, OriginNodes.Num());
+						ComponentInsertions.SetNum(ComponentOrigins.Num());
 						for (int32 i = 0; i < InInsertionIndices.Num(); i++)
 						{
-							if (!ComponentToIndex.Contains(ComponentIndex[InInsertionIndices[i]]))
+							if (!ComponentToMuscleIndex.Contains(ComponentIndex[InInsertionIndices[i]]))
 							{
 								ensureMsgf(false, TEXT("No origin in this component"));
 							}
-							else if (InsertionNodes[ComponentToIndex[ComponentIndex[InInsertionIndices[i]]]] == INDEX_NONE)
+							else
 							{
-								InsertionNodes[ComponentToIndex[ComponentIndex[InInsertionIndices[i]]]] = InInsertionIndices[i];
+								ComponentInsertions[ComponentToMuscleIndex[ComponentIndex[InInsertionIndices[i]]]].Add(InInsertionIndices[i]);
 							}
 						}
-						MuscleActivationElements.SetNum(OriginNodes.Num());
+						MuscleActivationElements.SetNum(ComponentOrigins.Num());
 						for (int32 ElemIdx = 0; ElemIdx < Elements->Num(); ElemIdx++)
 						{
-							if (ComponentToIndex.Contains(ComponentIndex[(*Elements)[ElemIdx][0]]))
+							if (ComponentToMuscleIndex.Contains(ComponentIndex[(*Elements)[ElemIdx][0]]))
 							{
-								MuscleActivationElements[ComponentToIndex[ComponentIndex[(*Elements)[ElemIdx][0]]]].Add(ElemIdx);
+								MuscleActivationElements[ComponentToMuscleIndex[ComponentIndex[(*Elements)[ElemIdx][0]]]].Add(ElemIdx);
 							}
 						}
-						OriginInsertionRestLength.SetNum(OriginNodes.Num());
-
-						for (int32 MuscleComponentIdx = 0; MuscleComponentIdx < OriginNodes.Num(); MuscleComponentIdx++)
+						//Choose one origin-insertion pair per muscle with largest distance apart
+						//use origin-insertion line segment length to estimate activation
+						for (int32 MuscleComponentIdx = 0; MuscleComponentIdx < ComponentOrigins.Num(); MuscleComponentIdx++)
 						{
-							if (ensureMsgf(InsertionNodes[MuscleComponentIdx] != INDEX_NONE, TEXT("InsertionNodes[%d] is not a pair"), MuscleComponentIdx))
+							if (ensureMsgf(ComponentOrigins.Num() > 0 && ComponentInsertions.Num() > 0, TEXT("Origin or Insertion missing in the muscle %d"), MuscleComponentIdx))
 							{
 								GeometryCollection::Facades::FMuscleActivationData MuscleActivationData;
+								MuscleActivationData.OriginInsertionRestLength = 0;
+								for (int32 OriginIdx : ComponentOrigins[MuscleComponentIdx])
+								{
+									for (int32 InsertionIdx : ComponentInsertions[MuscleComponentIdx])
+									{
+										float Dist = ((*Vertex)[OriginIdx] - (*Vertex)[InsertionIdx]).Size();
+										if (Dist > MuscleActivationData.OriginInsertionRestLength)
+										{
+											MuscleActivationData.OriginInsertionPair = FIntVector2(OriginIdx, InsertionIdx);
+											MuscleActivationData.OriginInsertionRestLength = Dist;
+										}
+									}
+								}
 								MuscleActivationData.MuscleActivationElement = MuscleActivationElements[MuscleComponentIdx];
-								MuscleActivationData.OriginInsertionPair = FIntVector2(OriginNodes[MuscleComponentIdx], InsertionNodes[MuscleComponentIdx]);
-								MuscleActivationData.OriginInsertionRestLength = ((*Vertex)[OriginNodes[MuscleComponentIdx]] - (*Vertex)[InsertionNodes[MuscleComponentIdx]]).Size();
 								MuscleActivationData.FiberDirectionMatrix.SetNum(MuscleActivationElements[MuscleComponentIdx].Num());
+								MuscleActivationData.ContractionVolumeScale.SetNum(MuscleActivationElements[MuscleComponentIdx].Num());
 								for (int32 LocalElemIdx = 0; LocalElemIdx < MuscleActivationElements[MuscleComponentIdx].Num(); LocalElemIdx++)
 								{
 									FVector3f V = (*FiberDirections)[MuscleActivationElements[MuscleComponentIdx][LocalElemIdx]];
-									// QR decomposition on vvT
+									// QR decomposition on vvT for orthogonal directions
 									FVector3f W = V;
 									if (V.X < V.Y)
 									{
@@ -87,7 +104,21 @@ void FComputeMuscleActivationDataNode::Evaluate(Dataflow::FContext& Context, con
 									FVector3f U = (V ^ W).GetSafeNormal();
 									W = (U ^ V).GetSafeNormal();
 									MuscleActivationData.FiberDirectionMatrix[LocalElemIdx] = Chaos::PMatrix33d(V, W, U);
+									//Muscle contraction volume scale
+									MuscleActivationData.ContractionVolumeScale[LocalElemIdx] = ContractionVolumeScale;
+									if (ContractionVolumeScalePerVertex)
+									{
+										float AverageScale = 1.f;
+										for (int32 ie = 0; ie < 4; ie++)
+										{
+											AverageScale += (*ContractionVolumeScalePerVertex)[(*Elements)[MuscleActivationElements[MuscleComponentIdx][LocalElemIdx]][ie]];
+											AverageScale /= 4.f;
+										}
+										MuscleActivationData.ContractionVolumeScale[LocalElemIdx] *= AverageScale;
+									}
 								}
+								
+								
 								FMuscleActivation.AddMuscleActivationData(MuscleActivationData);
 							}
 						}
diff --git a/Engine/Plugins/Experimental/ChaosFlesh/Source/ChaosFleshNodes/Public/Dataflow/ChaosFleshComputeMuscleActivationNode.h b/Engine/Plugins/Experimental/ChaosFlesh/Source/ChaosFleshNodes/Public/Dataflow/ChaosFleshComputeMuscleActivationNode.h
--- a/Engine/Plugins/Experimental/ChaosFlesh/Source/ChaosFleshNodes/Public/Dataflow/ChaosFleshComputeMuscleActivationNode.h
+++ b/Engine/Plugins/Experimental/ChaosFlesh/Source/ChaosFleshNodes/Public/Dataflow/ChaosFleshComputeMuscleActivationNode.h
@@ -28,6 +28,9 @@ public:
 	UPROPERTY(meta = (DataflowInput, DisplayName = "InsertionVertexIndices"))
 		TArray<int32> InsertionIndicesIn;
 
+	UPROPERTY(EditAnywhere, Category = "Dataflow")
+		float ContractionVolumeScale = 1.f;
+
 	FComputeMuscleActivationDataNode(const Dataflow::FNodeParameters& InParam, FGuid InGuid = FGuid::NewGuid())
 		: FDataflowNode(InParam, InGuid)
 	{
diff --git a/Engine/Source/Runtime/Experimental/Chaos/Private/GeometryCollection/Facades/CollectionMuscleActivationFacade.cpp b/Engine/Source/Runtime/Experimental/Chaos/Private/GeometryCollection/Facades/CollectionMuscleActivationFacade.cpp
--- a/Engine/Source/Runtime/Experimental/Chaos/Private/GeometryCollection/Facades/CollectionMuscleActivationFacade.cpp
+++ b/Engine/Source/Runtime/Experimental/Chaos/Private/GeometryCollection/Facades/CollectionMuscleActivationFacade.cpp
@@ -15,15 +15,16 @@ namespace GeometryCollection::Facades
 	const FName FMuscleActivationFacade::OriginInsertionPair("OriginInsertionPair");
 	const FName FMuscleActivationFacade::OriginInsertionRestLength("OriginInsertionRestLength");
 	const FName FMuscleActivationFacade::FiberDirectionMatrix("FiberDirectionMatrix");
-
+	const FName FMuscleActivationFacade::ContractionVolumeScale("ContractionVolumeScale");
 
 	FMuscleActivationFacade::FMuscleActivationFacade(FManagedArrayCollection& InCollection)
 		: ConstCollection(InCollection)
 		, Collection(&InCollection)
-		, MuscleActivationElementAttribute(InCollection, MuscleActivationElement, GroupName)
-		, OriginInsertionPairAttribute(InCollection, OriginInsertionPair, GroupName)
+		, MuscleActivationElementAttribute(InCollection, MuscleActivationElement, GroupName, "Tetrahedral")
+		, OriginInsertionPairAttribute(InCollection, OriginInsertionPair, GroupName, FGeometryCollection::VerticesGroup)
 		, OriginInsertionRestLengthAttribute(InCollection, OriginInsertionRestLength, GroupName)
 		, FiberDirectionMatrixAttribute(InCollection, FiberDirectionMatrix, GroupName)
+		, ContractionVolumeScaleAttribute(InCollection, ContractionVolumeScale, GroupName)
 	{
 		DefineSchema();
 	}
@@ -35,13 +36,15 @@ namespace GeometryCollection::Facades
 		, OriginInsertionPairAttribute(InCollection, OriginInsertionPair, GroupName)
 		, OriginInsertionRestLengthAttribute(InCollection, OriginInsertionRestLength, GroupName)
 		, FiberDirectionMatrixAttribute(InCollection, FiberDirectionMatrix, GroupName)
+		, ContractionVolumeScaleAttribute(InCollection, ContractionVolumeScale, GroupName)
 	{
 		
 	}
 
 	bool FMuscleActivationFacade::IsValid() const
 	{
-		return MuscleActivationElementAttribute.IsValid() && OriginInsertionPairAttribute.IsValid() && OriginInsertionRestLengthAttribute.IsValid() && FiberDirectionMatrixAttribute.IsValid();
+		return MuscleActivationElementAttribute.IsValid() && OriginInsertionPairAttribute.IsValid() && 
+			OriginInsertionRestLengthAttribute.IsValid() && FiberDirectionMatrixAttribute.IsValid() && ContractionVolumeScaleAttribute.IsValid();
 	}
 
 	void FMuscleActivationFacade::DefineSchema()
@@ -51,6 +54,7 @@ namespace GeometryCollection::Facades
 		OriginInsertionPairAttribute.Add(ManageArrayAccessor::EPersistencePolicy::MakePersistent, FGeometryCollection::VerticesGroup);
 		OriginInsertionRestLengthAttribute.Add();
 		FiberDirectionMatrixAttribute.Add();
+		ContractionVolumeScaleAttribute.Add();
 	}
 
 	int32 FMuscleActivationFacade::AddMuscleActivationData(const FMuscleActivationData& InputData)
@@ -63,6 +67,7 @@ namespace GeometryCollection::Facades
 			OriginInsertionPairAttribute.Modify()[NewIndex] = InputData.OriginInsertionPair;
 			OriginInsertionRestLengthAttribute.Modify()[NewIndex] = InputData.OriginInsertionRestLength;
 			FiberDirectionMatrixAttribute.Modify()[NewIndex] = InputData.FiberDirectionMatrix;
+			ContractionVolumeScaleAttribute.Modify()[NewIndex] = InputData.ContractionVolumeScale;
 			return NewIndex;
 		}
 		return INDEX_NONE;
@@ -79,6 +84,7 @@ namespace GeometryCollection::Facades
 				ReturnData.OriginInsertionPair = OriginInsertionPairAttribute.Get()[DataIndex];
 				ReturnData.OriginInsertionRestLength = OriginInsertionRestLengthAttribute.Get()[DataIndex];
 				ReturnData.FiberDirectionMatrix = FiberDirectionMatrixAttribute.Get()[DataIndex];
+				ReturnData.ContractionVolumeScale = ContractionVolumeScaleAttribute.Get()[DataIndex];
 			}
 		}
 		return ReturnData;
diff --git a/Engine/Source/Runtime/Experimental/Chaos/Public/Chaos/Deformable/MuscleActivationConstraints.h b/Engine/Source/Runtime/Experimental/Chaos/Public/Chaos/Deformable/MuscleActivationConstraints.h
--- a/Engine/Source/Runtime/Experimental/Chaos/Public/Chaos/Deformable/MuscleActivationConstraints.h
+++ b/Engine/Source/Runtime/Experimental/Chaos/Public/Chaos/Deformable/MuscleActivationConstraints.h
@@ -39,6 +39,7 @@ namespace Chaos::Softs
 				OriginInsertionRestLength.Add(MuscleActivationData.OriginInsertionRestLength);
 				Activation.Add(0.f);
 				FiberDirectionMatrix.Add(MuscleActivationData.FiberDirectionMatrix);
+				ContractionVolumeScale.Add(MuscleActivationData.ContractionVolumeScale);
 			}
 		}
 
@@ -67,7 +68,7 @@ namespace Chaos::Softs
 				{
 					for (int32 ElemIdx = 0; ElemIdx < MuscleActivationElement[MuscleIdx].Num(); ElemIdx++)
 					{
-						Constraints.ModifyDmInverseFromFiberLength(MuscleActivationElement[MuscleIdx][ElemIdx], Activation[MuscleIdx], FiberDirectionMatrix[MuscleIdx][ElemIdx]);
+						Constraints.ModifyDmInverseFromFiberLength(MuscleActivationElement[MuscleIdx][ElemIdx], Activation[MuscleIdx], FiberDirectionMatrix[MuscleIdx][ElemIdx], ContractionVolumeScale[MuscleIdx][ElemIdx]);
 					}
 				}
 			}
@@ -78,6 +79,7 @@ namespace Chaos::Softs
 		TArray<float> OriginInsertionRestLength;
 		TArray<float> Activation;
 		TArray<TArray<Chaos::PMatrix33d>> FiberDirectionMatrix;
+		TArray<TArray<float>> ContractionVolumeScale;
 	};
 
 
diff --git a/Engine/Source/Runtime/Experimental/Chaos/Public/Chaos/XPBDCorotatedConstraints.h b/Engine/Source/Runtime/Experimental/Chaos/Public/Chaos/XPBDCorotatedConstraints.h
--- a/Engine/Source/Runtime/Experimental/Chaos/Public/Chaos/XPBDCorotatedConstraints.h
+++ b/Engine/Source/Runtime/Experimental/Chaos/Public/Chaos/XPBDCorotatedConstraints.h
@@ -406,13 +406,13 @@ namespace Chaos::Softs
 			return dC2;
 		}
 
-		void ModifyDmInverseFromFiberLength(const int32 e, const float L, const PMatrix<T, 3, 3>& MFiberDir) const
+		void ModifyDmInverseFromFiberLength(const int32 e, const T L, const PMatrix<T, 3, 3>& MFiberDir, const T ContractionVolumeScale) const
 		{
 			TRACE_CPUPROFILER_EVENT_SCOPE(TEXT("STAT_ChaosXPBDCorotatedModifyDmInverseFromFiberLength"));
 			PMatrix<T, 3, 3> DmInv = ElementDmInvSave(e);
 			if (L < 1)
 			{
-				PMatrix<T, 3, 3> D(1 / L, sqrt(L), sqrt(L));
+				PMatrix<T, 3, 3> D(1 / L, FMath::Pow(L, ContractionVolumeScale/T(2)), FMath::Pow(L, ContractionVolumeScale/T(2)));
 				PMatrix<T, 3, 3> Factor = MFiberDir * D * MFiberDir.GetTransposed();
 				DmInv = Factor * DmInv; //Matrix multiplication convention: AB = B*A
 			}
diff --git a/Engine/Source/Runtime/Experimental/Chaos/Public/GeometryCollection/Facades/CollectionMuscleActivationFacade.h b/Engine/Source/Runtime/Experimental/Chaos/Public/GeometryCollection/Facades/CollectionMuscleActivationFacade.h
--- a/Engine/Source/Runtime/Experimental/Chaos/Public/GeometryCollection/Facades/CollectionMuscleActivationFacade.h
+++ b/Engine/Source/Runtime/Experimental/Chaos/Public/GeometryCollection/Facades/CollectionMuscleActivationFacade.h
@@ -15,6 +15,7 @@ namespace GeometryCollection::Facades
 		FIntVector2 OriginInsertionPair;
 		float OriginInsertionRestLength;
 		TArray<Chaos::PMatrix33d> FiberDirectionMatrix;
+		TArray<float> ContractionVolumeScale;
 	};
 
 	/** Kinematic Facade */
@@ -27,6 +28,7 @@ namespace GeometryCollection::Facades
 		static CHAOS_API const FName OriginInsertionPair;
 		static CHAOS_API const FName OriginInsertionRestLength;
 		static CHAOS_API const FName FiberDirectionMatrix;
+		static CHAOS_API const FName ContractionVolumeScale;
 
 		CHAOS_API FMuscleActivationFacade(FManagedArrayCollection& InCollection);
 		CHAOS_API FMuscleActivationFacade(const FManagedArrayCollection& InCollection);
@@ -55,5 +57,6 @@ namespace GeometryCollection::Facades
 		TManagedArrayAccessor<FIntVector2> OriginInsertionPairAttribute;
 		TManagedArrayAccessor<float> OriginInsertionRestLengthAttribute;
 		TManagedArrayAccessor<TArray<Chaos::PMatrix33d>> FiberDirectionMatrixAttribute;
+		TManagedArrayAccessor<TArray<float>> ContractionVolumeScaleAttribute;
 	};
 }
