ChaosFlesh : Added call to Super::EndPlay in the TetrahedralComponent.

diff --git a/Engine/Plugins/Experimental/ChaosFlesh/Source/ChaosFleshEngine/Private/ChaosFlesh/ChaosDeformableTetrahedralComponent.cpp b/Engine/Plugins/Experimental/ChaosFlesh/Source/ChaosFleshEngine/Private/ChaosFlesh/ChaosDeformableTetrahedralComponent.cpp
--- a/Engine/Plugins/Experimental/ChaosFlesh/Source/ChaosFleshEngine/Private/ChaosFlesh/ChaosDeformableTetrahedralComponent.cpp
+++ b/Engine/Plugins/Experimental/ChaosFlesh/Source/ChaosFleshEngine/Private/ChaosFlesh/ChaosDeformableTetrahedralComponent.cpp
@@ -83,6 +83,8 @@ void UDeformableTetrahedralComponent::EndPlay(const EEndPlayReason::Type ReasonE
 	{
 		GetSimulationCollection()->Reset();
 	}
+
+	Super::EndPlay(ReasonEnd);
 }
 
 void UDeformableTetrahedralComponent::SetRestCollection(const UFleshAsset* InRestCollection)
