ChaosFlesh : FleshDeformer Cleanup

diff --git a/Engine/Plugins/Experimental/ChaosFlesh/Source/ChaosFleshNodes/Private/Dataflow/ChaosFleshGenerateSurfaceBindingsNode.cpp b/Engine/Plugins/Experimental/ChaosFlesh/Source/ChaosFleshNodes/Private/Dataflow/ChaosFleshGenerateSurfaceBindingsNode.cpp
--- a/Engine/Plugins/Experimental/ChaosFlesh/Source/ChaosFleshNodes/Private/Dataflow/ChaosFleshGenerateSurfaceBindingsNode.cpp
+++ b/Engine/Plugins/Experimental/ChaosFlesh/Source/ChaosFleshNodes/Private/Dataflow/ChaosFleshGenerateSurfaceBindingsNode.cpp
@@ -305,332 +305,334 @@ FGenerateSurfaceBindings::Evaluate(Dataflow::FContext& Context, const FDataflowO
 				}
 				const int32 TetMeshStart = (*TetrahedronStart)[TetMeshIdx];
 				const int32 TetMeshCount = (*TetrahedronCount)[TetMeshIdx];
-
-				// Build Tetrahedra
-				TArray<Chaos::TTetrahedron<Chaos::FReal>> Tets;			// Index 0 == TetMeshStart
-				TArray<Chaos::TTetrahedron<Chaos::FReal>*> BVHTetPtrs;
-				Tets.SetNumUninitialized(TetMeshCount);
-				BVHTetPtrs.SetNumUninitialized(TetMeshCount);
-				for (int32 i = 0; i < TetMeshCount; i++)
-				{
-					const int32 Idx = TetMeshStart + i;
-					const FIntVector4& Tet = (*Tetrahedron)[Idx];
-					Tets[i] = Chaos::TTetrahedron<Chaos::FReal>(
-						(*Vertex)[Tet[0]],
-						(*Vertex)[Tet[1]],
-						(*Vertex)[Tet[2]],
-						(*Vertex)[Tet[3]]);
-					BVHTetPtrs[i] = &Tets[i];
-				}
-
-				// Init BVH for tetrahedra.
-				Chaos::TBoundingVolumeHierarchy<
-					TArray<Chaos::TTetrahedron<Chaos::FReal>*>, 
-					TArray<int32>, 
-					Chaos::FReal, 
-					3> TetBVH(BVHTetPtrs);
-
-				//
-				// Init boundary mesh for projections.
-				//
-
-				Chaos::FTriangleMesh TetBoundaryMesh;
-				Chaos::FTriangleMesh::TBVHType<Chaos::FRealDouble> TetBoundaryBVH;
-				TArray<Chaos::FVec3> VertexD;
-				TConstArrayView<Chaos::TVec3<Chaos::FRealDouble>> VertexDView(VertexD);
-				TArray<Chaos::FVec3> PointNormals;
-				const int32 TriMeshStart = (*FacesStart)[TetMeshIdx];
-				const int32 TriMeshCount = (*FacesCount)[TetMeshIdx];
-				if (bDoSurfaceProjection)
+				if (TetMeshCount > 0)
 				{
-					if (TriMeshCount == Triangle->GetConstArray().Num())
+					// Build Tetrahedra
+					TArray<Chaos::TTetrahedron<Chaos::FReal>> Tets;			// Index 0 == TetMeshStart
+					TArray<Chaos::TTetrahedron<Chaos::FReal>*> BVHTetPtrs;
+					Tets.SetNumUninitialized(TetMeshCount);
+					BVHTetPtrs.SetNumUninitialized(TetMeshCount);
+					for (int32 i = 0; i < TetMeshCount; i++)
 					{
-						TetBoundaryMesh.Init(
-							reinterpret_cast<const TArray<Chaos::TVec3<int32>>&>(Triangle->GetConstArray()), 0, -1, false);
-					}
-					else
-					{
-						TArray<Chaos::TVec3<int32>> Faces; Faces.SetNumUninitialized(TriMeshCount);
-						for (int32 i = 0; i < TriMeshCount; i++)
-						{
-							const int32 Idx = TriMeshStart + i;
-							Faces[i] = reinterpret_cast<const Chaos::TVec3<int32>&>(Triangle->GetConstArray()[Idx]);
-						}
-						TetBoundaryMesh.Init(Faces, 0, -1, false);
+						const int32 Idx = TetMeshStart + i;
+						const FIntVector4& Tet = (*Tetrahedron)[Idx];
+						Tets[i] = Chaos::TTetrahedron<Chaos::FReal>(
+							(*Vertex)[Tet[0]],
+							(*Vertex)[Tet[1]],
+							(*Vertex)[Tet[2]],
+							(*Vertex)[Tet[3]]);
+						BVHTetPtrs[i] = &Tets[i];
 					}
 
-					// Promote vertices to double because that's what FTriangleMesh wants.
-					VertexD.SetNumUninitialized(Vertex->Num());
-					for (int32 i = 0; i < VertexD.Num(); i++)
-					{
-						VertexD[i] = Chaos::FVec3((*Vertex)[i][0], (*Vertex)[i][1], (*Vertex)[i][2]);
-					}
-					VertexDView = TConstArrayView<Chaos::TVec3<Chaos::FRealDouble>>(VertexD);
+					// Init BVH for tetrahedra.
+					Chaos::TBoundingVolumeHierarchy<
+						TArray<Chaos::TTetrahedron<Chaos::FReal>*>,
+						TArray<int32>,
+						Chaos::FReal,
+						3> TetBVH(BVHTetPtrs);
 
-					PointNormals = TetBoundaryMesh.GetPointNormals(VertexDView, false, true);
-					TetBoundaryMesh.BuildBVH(VertexDView, TetBoundaryBVH);
-				}
-
-				//
-				// Do intersection tests against tets, then the surface.
-				//
-
-				TArray<TArray<FIntVector4>> Parents; Parents.SetNum(MeshVertices.Num());
-				TArray<TArray<FVector4f>> Weights;	 Weights.SetNum(MeshVertices.Num());
-				TArray<TArray<FVector3f>> Offsets;	 Offsets.SetNum(MeshVertices.Num());
-				TArray<TArray<float>> Masks;		 Masks.SetNum(MeshVertices.Num());
-				TArray<int32> Orphans;
-				int32 TetHits = 0;
-				int32 TriHits = 0;
-				int32 Adoptions = 0;
-				int32 NumOrphans = 0;
-				for (int32 LOD = 0; LOD < MeshVertices.Num(); LOD++)
-				{
-					Parents[LOD].SetNumUninitialized(MeshVertices[LOD].Num());
-					Weights[LOD].SetNumUninitialized(MeshVertices[LOD].Num());
-					Offsets[LOD].SetNumUninitialized(MeshVertices[LOD].Num());
-					Masks[LOD].SetNumUninitialized(MeshVertices[LOD].Num());
+					//
+					// Init boundary mesh for projections.
+					//
 
-					TArray<int32> TetIntersections; TetIntersections.Reserve(64);
-					for (int32 i = 0; i < MeshVertices[LOD].Num(); i++)
+					Chaos::FTriangleMesh TetBoundaryMesh;
+					Chaos::FTriangleMesh::TBVHType<Chaos::FRealDouble> TetBoundaryBVH;
+					TArray<Chaos::FVec3> VertexD;
+					TConstArrayView<Chaos::TVec3<Chaos::FRealDouble>> VertexDView(VertexD);
+					TArray<Chaos::FVec3> PointNormals;
+					const int32 TriMeshStart = (*FacesStart)[TetMeshIdx];
+					const int32 TriMeshCount = (*FacesCount)[TetMeshIdx];
+					if (bDoSurfaceProjection)
 					{
-						Parents[LOD][i] = FIntVector4(INDEX_NONE);
-						Weights[LOD][i] = FVector4f(0);
-						Offsets[LOD][i] = FVector3f(0);
-						Masks[LOD][i] = 0.0; // Shader does skinning for this vertex
-
-						const FVector3f& Pos = MeshVertices[LOD][i];
-						Chaos::TVec3<Chaos::FReal> PosD(Pos[0], Pos[1], Pos[2]);
-						TetIntersections = TetBVH.FindAllIntersections(PosD);
-						int32 j = 0;
-						for (; j < TetIntersections.Num(); j++)
+						if (TriMeshCount == Triangle->GetConstArray().Num())
 						{
-							const int32 TetIdx = TetIntersections[j];
-							if (!Tets[TetIdx].Outside(Pos, 1.0e-2)) // includes boundary
-							{
-								TetHits++;
-								Parents[LOD][i] = (*Tetrahedron)[TetIdx + TetMeshStart];
-								Chaos::TVector<Chaos::FReal, 4> WeightsD = Tets[TetIdx].GetBarycentricCoordinates(Pos);
-								Weights[LOD][i] = FVector4f(WeightsD[0], WeightsD[1], WeightsD[2], WeightsD[3]);
-								Offsets[LOD][i] = FVector3f(0);
-								Masks[LOD][i] = 1.0; // Shader does sim for this vertex
-
-								FVector3f EmbeddedPos =
-									(*Vertex)[Parents[LOD][i][0]] * Weights[LOD][i][0] +
-									(*Vertex)[Parents[LOD][i][1]] * Weights[LOD][i][1] +
-									(*Vertex)[Parents[LOD][i][2]] * Weights[LOD][i][2] +
-									(*Vertex)[Parents[LOD][i][3]] * Weights[LOD][i][3];
-								check((Pos - EmbeddedPos).SquaredLength() < 1.0);
-
-								break;
-							}
+							TetBoundaryMesh.Init(
+								reinterpret_cast<const TArray<Chaos::TVec3<int32>>&>(Triangle->GetConstArray()), 0, -1, false);
 						}
-						if (j == TetIntersections.Num())
+						else
 						{
-							// This vertex didn't land inside any tetrahedra. Project it to the tet boundary surface.
-							int32 TriIdx = INDEX_NONE;
-							Chaos::FVec3 TriWeights;
-							if (bDoSurfaceProjection && 
-								TetBoundaryMesh.SmoothProject(
-									TetBoundaryBVH,
-									VertexDView,
-									PointNormals, Pos,
-									TriIdx, TriWeights, SurfaceProjectionIterations))
+							TArray<Chaos::TVec3<int32>> Faces; Faces.SetNumUninitialized(TriMeshCount);
+							for (int32 i = 0; i < TriMeshCount; i++)
 							{
-								TriHits++;
-								const FIntVector& Tri = (*Triangle)[TriIdx + TriMeshStart];
-								Parents[LOD][i][0] = Tri[0];
-								Parents[LOD][i][1] = Tri[1];
-								Parents[LOD][i][2] = Tri[2];
-								Parents[LOD][i][3] = INDEX_NONE;
-
-								Weights[LOD][i][0] = TriWeights[0];
-								Weights[LOD][i][1] = TriWeights[1];
-								Weights[LOD][i][2] = TriWeights[2];
-								Weights[LOD][i][3] = 0.0;
-
-								const FVector3f EmbeddedPos =
-									TriWeights[0] * Vertex->GetConstArray()[Tri[0]] +
-									TriWeights[1] * Vertex->GetConstArray()[Tri[1]] +
-									TriWeights[2] * Vertex->GetConstArray()[Tri[2]];
-								Offsets[LOD][i] = EmbeddedPos - Pos;
-
-								Masks[LOD][i] = 1.0; // Shader does sim for this vertex
+								const int32 Idx = TriMeshStart + i;
+								Faces[i] = reinterpret_cast<const Chaos::TVec3<int32>&>(Triangle->GetConstArray()[Idx]);
 							}
-							else
-							{
-								// Despair...
-								Orphans.Add(i);
-
-								Parents[LOD][i][0] = INDEX_NONE;
-								Parents[LOD][i][1] = INDEX_NONE;
-								Parents[LOD][i][2] = INDEX_NONE;
-								Parents[LOD][i][3] = INDEX_NONE;
-
-								Weights[LOD][i][0] = 0.0;
-								Weights[LOD][i][1] = 0.0;
-								Weights[LOD][i][2] = 0.0;
-								Weights[LOD][i][3] = 0.0;
-
-								Offsets[LOD][i][0] = 0.0;
-								Offsets[LOD][i][1] = 0.0;
-								Offsets[LOD][i][2] = 0.0;
-
-								Masks[LOD][i] = 0.0; // Shader does skinning for this vertex
-							} // if !SmoothProject()
-						} // if !TetIntersections
-					} // end for all vertices
-
-					// 
-					// Advancing front orphan reparenting
-					//
-					if (!MeshNeighborNodes.IsValidIndex(LOD))
-					{
-						continue;
+							TetBoundaryMesh.Init(Faces, 0, -1, false);
+						}
+
+						// Promote vertices to double because that's what FTriangleMesh wants.
+						VertexD.SetNumUninitialized(Vertex->Num());
+						for (int32 i = 0; i < VertexD.Num(); i++)
+						{
+							VertexD[i] = Chaos::FVec3((*Vertex)[i][0], (*Vertex)[i][1], (*Vertex)[i][2]);
+						}
+						VertexDView = TConstArrayView<Chaos::TVec3<Chaos::FRealDouble>>(VertexD);
+
+						PointNormals = TetBoundaryMesh.GetPointNormals(VertexDView, false, true);
+						TetBoundaryMesh.BuildBVH(VertexDView, TetBoundaryBVH);
 					}
-					const TArray<TArray<uint32>>& NeighborNodes = MeshNeighborNodes[LOD];
-					TSet<int32> OrphanSet(Orphans);
-					while (bDoOrphanReparenting && Orphans.Num())
+
+					//
+					// Do intersection tests against tets, then the surface.
+					//
+
+					TArray<TArray<FIntVector4>> Parents; Parents.SetNum(MeshVertices.Num());
+					TArray<TArray<FVector4f>> Weights;	 Weights.SetNum(MeshVertices.Num());
+					TArray<TArray<FVector3f>> Offsets;	 Offsets.SetNum(MeshVertices.Num());
+					TArray<TArray<float>> Masks;		 Masks.SetNum(MeshVertices.Num());
+					TArray<int32> Orphans;
+					int32 TetHits = 0;
+					int32 TriHits = 0;
+					int32 Adoptions = 0;
+					int32 NumOrphans = 0;
+					for (int32 LOD = 0; LOD < MeshVertices.Num(); LOD++)
 					{
-						// Find the orphan with the fewest number of orphan neighbors, and the 
-						// most non-orphans in their 1 ring.
-						int32 Orphan = INDEX_NONE;
-						int32 NumOrphanNeighbors = TNumericLimits<int32>::Max();
-						int32 NumNonOrphanNeighbors = 0;
-						for (int32 i = 0; i < Orphans.Num(); i++)
+						Parents[LOD].SetNumUninitialized(MeshVertices[LOD].Num());
+						Weights[LOD].SetNumUninitialized(MeshVertices[LOD].Num());
+						Offsets[LOD].SetNumUninitialized(MeshVertices[LOD].Num());
+						Masks[LOD].SetNumUninitialized(MeshVertices[LOD].Num());
+
+						TArray<int32> TetIntersections; TetIntersections.Reserve(64);
+						for (int32 i = 0; i < MeshVertices[LOD].Num(); i++)
 						{
-							int32 CurrOrphan = Orphans[i];
-							if (!NeighborNodes.IsValidIndex(CurrOrphan))
+							Parents[LOD][i] = FIntVector4(INDEX_NONE);
+							Weights[LOD][i] = FVector4f(0);
+							Offsets[LOD][i] = FVector3f(0);
+							Masks[LOD][i] = 0.0; // Shader does skinning for this vertex
+
+							const FVector3f& Pos = MeshVertices[LOD][i];
+							Chaos::TVec3<Chaos::FReal> PosD(Pos[0], Pos[1], Pos[2]);
+							TetIntersections = TetBVH.FindAllIntersections(PosD);
+							int32 j = 0;
+							for (; j < TetIntersections.Num(); j++)
 							{
-								continue;
+								const int32 TetIdx = TetIntersections[j];
+								if (!Tets[TetIdx].Outside(Pos, 1.0e-2)) // includes boundary
+								{
+									TetHits++;
+									Parents[LOD][i] = (*Tetrahedron)[TetIdx + TetMeshStart];
+									Chaos::TVector<Chaos::FReal, 4> WeightsD = Tets[TetIdx].GetBarycentricCoordinates(Pos);
+									Weights[LOD][i] = FVector4f(WeightsD[0], WeightsD[1], WeightsD[2], WeightsD[3]);
+									Offsets[LOD][i] = FVector3f(0);
+									Masks[LOD][i] = 1.0; // Shader does sim for this vertex
+
+									FVector3f EmbeddedPos =
+										(*Vertex)[Parents[LOD][i][0]] * Weights[LOD][i][0] +
+										(*Vertex)[Parents[LOD][i][1]] * Weights[LOD][i][1] +
+										(*Vertex)[Parents[LOD][i][2]] * Weights[LOD][i][2] +
+										(*Vertex)[Parents[LOD][i][3]] * Weights[LOD][i][3];
+									check((Pos - EmbeddedPos).SquaredLength() < 1.0);
+
+									break;
+								}
 							}
-							const TArray<uint32>& Neighbors = NeighborNodes[CurrOrphan];
-							int32 OrphanCount = 0;
-							int32 NonOrphanCount = 0;
-							for (int32 j = 0; j < Neighbors.Num(); j++)
+							if (j == TetIntersections.Num())
 							{
-								if (OrphanSet.Contains(Neighbors[j]))
+								// This vertex didn't land inside any tetrahedra. Project it to the tet boundary surface.
+								int32 TriIdx = INDEX_NONE;
+								Chaos::FVec3 TriWeights;
+								if (bDoSurfaceProjection &&
+									TetBoundaryMesh.SmoothProject(
+										TetBoundaryBVH,
+										VertexDView,
+										PointNormals, Pos,
+										TriIdx, TriWeights, SurfaceProjectionIterations))
 								{
-									OrphanCount++;
+									TriHits++;
+									const FIntVector& Tri = (*Triangle)[TriIdx + TriMeshStart];
+									Parents[LOD][i][0] = Tri[0];
+									Parents[LOD][i][1] = Tri[1];
+									Parents[LOD][i][2] = Tri[2];
+									Parents[LOD][i][3] = INDEX_NONE;
+
+									Weights[LOD][i][0] = TriWeights[0];
+									Weights[LOD][i][1] = TriWeights[1];
+									Weights[LOD][i][2] = TriWeights[2];
+									Weights[LOD][i][3] = 0.0;
+
+									const FVector3f EmbeddedPos =
+										TriWeights[0] * Vertex->GetConstArray()[Tri[0]] +
+										TriWeights[1] * Vertex->GetConstArray()[Tri[1]] +
+										TriWeights[2] * Vertex->GetConstArray()[Tri[2]];
+									Offsets[LOD][i] = EmbeddedPos - Pos;
+
+									Masks[LOD][i] = 1.0; // Shader does sim for this vertex
 								}
 								else
 								{
-									NonOrphanCount++;
-								}
-							}
-							if (OrphanCount <= NumOrphanNeighbors && NonOrphanCount > NumNonOrphanNeighbors)
-							{
-								Orphan = CurrOrphan;
-								NumOrphanNeighbors = OrphanCount;
-								NumNonOrphanNeighbors = NonOrphanCount;
-							}
-						}
-						if (Orphan == INDEX_NONE)
+									// Despair...
+									Orphans.Add(i);
+
+									Parents[LOD][i][0] = INDEX_NONE;
+									Parents[LOD][i][1] = INDEX_NONE;
+									Parents[LOD][i][2] = INDEX_NONE;
+									Parents[LOD][i][3] = INDEX_NONE;
+
+									Weights[LOD][i][0] = 0.0;
+									Weights[LOD][i][1] = 0.0;
+									Weights[LOD][i][2] = 0.0;
+									Weights[LOD][i][3] = 0.0;
+
+									Offsets[LOD][i][0] = 0.0;
+									Offsets[LOD][i][1] = 0.0;
+									Offsets[LOD][i][2] = 0.0;
+
+									Masks[LOD][i] = 0.0; // Shader does skinning for this vertex
+								} // if !SmoothProject()
+							} // if !TetIntersections
+						} // end for all vertices
+
+						// 
+						// Advancing front orphan reparenting
+						//
+						if (!MeshNeighborNodes.IsValidIndex(LOD))
 						{
-							// We only have orphans with no neighbors left.
-							break;
+							continue;
 						}
-						const FVector3f& Pos = MeshVertices[LOD][Orphan];
-						Chaos::TVec3<Chaos::FReal> PosD(Pos[0], Pos[1], Pos[2]);
-
-						// Use the parent simplices of non-orphan neighbors as test candidates.
-						Chaos::FReal CurrDist = TNumericLimits<Chaos::FReal>::Max();
-						const TArray<uint32>& Neighbors = NeighborNodes[Orphan];
-						bool FoundBinding = false;
-						for (int32 i = 0; i < Neighbors.Num(); i++)
+						const TArray<TArray<uint32>>& NeighborNodes = MeshNeighborNodes[LOD];
+						TSet<int32> OrphanSet(Orphans);
+						while (bDoOrphanReparenting && Orphans.Num())
 						{
-							const uint32 Neighbor = Neighbors[i];
-							if (OrphanSet.Contains(Neighbor))
-							{
-								continue;
-							}
-
-							const FIntVector4& P = Parents[LOD][Neighbor];
-							int32 NumValid = 0;
-							for (int32 j = 0; j < 4; j++)
+							// Find the orphan with the fewest number of orphan neighbors, and the 
+							// most non-orphans in their 1 ring.
+							int32 Orphan = INDEX_NONE;
+							int32 NumOrphanNeighbors = TNumericLimits<int32>::Max();
+							int32 NumNonOrphanNeighbors = 0;
+							for (int32 i = 0; i < Orphans.Num(); i++)
 							{
-								NumValid += P[j] != INDEX_NONE ? 1 : 0;
+								int32 CurrOrphan = Orphans[i];
+								if (!NeighborNodes.IsValidIndex(CurrOrphan))
+								{
+									continue;
+								}
+								const TArray<uint32>& Neighbors = NeighborNodes[CurrOrphan];
+								int32 OrphanCount = 0;
+								int32 NonOrphanCount = 0;
+								for (int32 j = 0; j < Neighbors.Num(); j++)
+								{
+									if (OrphanSet.Contains(Neighbors[j]))
+									{
+										OrphanCount++;
+									}
+									else
+									{
+										NonOrphanCount++;
+									}
+								}
+								if (OrphanCount <= NumOrphanNeighbors && NonOrphanCount > NumNonOrphanNeighbors)
+								{
+									Orphan = CurrOrphan;
+									NumOrphanNeighbors = OrphanCount;
+									NumNonOrphanNeighbors = NonOrphanCount;
+								}
 							}
-
-							if (NumValid == 0)
+							if (Orphan == INDEX_NONE)
 							{
-								continue;
+								// We only have orphans with no neighbors left.
+								break;
 							}
-							else 
+							const FVector3f& Pos = MeshVertices[LOD][Orphan];
+							Chaos::TVec3<Chaos::FReal> PosD(Pos[0], Pos[1], Pos[2]);
+
+							// Use the parent simplices of non-orphan neighbors as test candidates.
+							Chaos::FReal CurrDist = TNumericLimits<Chaos::FReal>::Max();
+							const TArray<uint32>& Neighbors = NeighborNodes[Orphan];
+							bool FoundBinding = false;
+							for (int32 i = 0; i < Neighbors.Num(); i++)
 							{
-								// Find tets that share parent indices
-								for(int32 j=0; j < 4; j++)
+								const uint32 Neighbor = Neighbors[i];
+								if (OrphanSet.Contains(Neighbor))
 								{
-									const int32 ParentIdx = P[j];
-									if (IncidentElements->GetConstArray().IsValidIndex(ParentIdx))
+									continue;
+								}
+
+								const FIntVector4& P = Parents[LOD][Neighbor];
+								int32 NumValid = 0;
+								for (int32 j = 0; j < 4; j++)
+								{
+									NumValid += P[j] != INDEX_NONE ? 1 : 0;
+								}
+
+								if (NumValid == 0)
+								{
+									continue;
+								}
+								else
+								{
+									// Find tets that share parent indices
+									for (int32 j = 0; j < 4; j++)
 									{
-										const TArray<int32>& NeighborTets = (*IncidentElements)[ParentIdx];
-										for (int32 k = 0; k < NeighborTets.Num(); k++)
+										const int32 ParentIdx = P[j];
+										if (IncidentElements->GetConstArray().IsValidIndex(ParentIdx))
 										{
-											const int32 TetIdx = NeighborTets[k] - TetMeshStart;
-											if (ensure(Tets.IsValidIndex(TetIdx)))
+											const TArray<int32>& NeighborTets = (*IncidentElements)[ParentIdx];
+											for (int32 k = 0; k < NeighborTets.Num(); k++)
 											{
-												const Chaos::FTetrahedron& Tet = Tets[TetIdx];
-
-												Chaos::TVec4<Chaos::FReal> W;
-												Chaos::TVec3<Chaos::FReal> EmbeddedPos = Tet.FindClosestPointAndBary(PosD, W, 1.0e-4);
-												Chaos::TVec3<Chaos::FReal> O = EmbeddedPos - PosD;
-												Chaos::FReal Dist = O.SquaredLength();
-												if (Dist < CurrDist)
+												const int32 TetIdx = NeighborTets[k] - TetMeshStart;
+												if (ensure(Tets.IsValidIndex(TetIdx)))
 												{
-													CurrDist = Dist;
-													Parents[LOD][Orphan] = (*Tetrahedron)[TetIdx + TetMeshStart];
-													Weights[LOD][Orphan] = FVector4f(W[0], W[1], W[2], W[3]);
-													Offsets[LOD][Orphan] = FVector3f(O[0], O[1], O[2]);
-													Masks[LOD][i] = 1.0; // Shader does sim for this vertex
-													FoundBinding = true;
+													const Chaos::FTetrahedron& Tet = Tets[TetIdx];
+
+													Chaos::TVec4<Chaos::FReal> W;
+													Chaos::TVec3<Chaos::FReal> EmbeddedPos = Tet.FindClosestPointAndBary(PosD, W, 1.0e-4);
+													Chaos::TVec3<Chaos::FReal> O = EmbeddedPos - PosD;
+													Chaos::FReal Dist = O.SquaredLength();
+													if (Dist < CurrDist)
+													{
+														CurrDist = Dist;
+														Parents[LOD][Orphan] = (*Tetrahedron)[TetIdx + TetMeshStart];
+														Weights[LOD][Orphan] = FVector4f(W[0], W[1], W[2], W[3]);
+														Offsets[LOD][Orphan] = FVector3f(O[0], O[1], O[2]);
+														Masks[LOD][i] = 1.0; // Shader does sim for this vertex
+														FoundBinding = true;
+													}
 												}
 											}
 										}
 									}
 								}
-							}
-						} // end for all neighbors
+							} // end for all neighbors
 
-						// Whether or not we successfully reparented, remove the orphan from the list.
-						OrphanSet.Remove(Orphan);
-						Orphans.Remove(Orphan);
-						if (FoundBinding)
-						{
-							Adoptions++;
-						}
-						else
-						{
-							NumOrphans++;
-						}
-					} // end while(Orphans)
-					NumOrphans += Orphans.Num();
-
-					//ELogVerbosity::Type Verbosity = Orphans.Num() > 0 ? ELogVerbosity::Error : ELogVerbosity::Display;
-					UE_LOG(LogMeshBindings, Display,
-						TEXT("'%s' - Generated mesh bindings between tet mesh index %d and %s mesh of '%s' LOD %d - stats:\n"
-							"    Render vertices num: %d\n"
-							"    Vertices in tetrahedra: %d\n"
-							"    Vertices bound to tet surface: %d\n"
-							"    Orphaned vertices reparented: %d\n"
-							"    Vertices orphaned: %d"),
-						*GetName().ToString(),
-						TetMeshIdx, 
-						bUseSkeletalMeshImportModel ? TEXT("import") : TEXT("render"),
-						*MeshId, LOD,
-						MeshVertices[LOD].Num(), TetHits, TriHits, Adoptions, NumOrphans);
-
-				} // end for all LOD
-
-				// Stash bindings in the geometry collection
-				GeometryCollection::Facades::FTetrahedralBindings TetBindings(OutCollection);
-				TetBindings.DefineSchema();
-				FName MeshName(*MeshId, MeshId.Len());
-				for (int32 LOD = 0; LOD < MeshVertices.Num(); LOD++)
-				{
-					TetBindings.AddBindingsGroup(TetMeshIdx, MeshName, LOD);
-					TetBindings.SetBindingsData(Parents[LOD], Weights[LOD], Offsets[LOD], Masks[LOD]);
-				}
+							// Whether or not we successfully reparented, remove the orphan from the list.
+							OrphanSet.Remove(Orphan);
+							Orphans.Remove(Orphan);
+							if (FoundBinding)
+							{
+								Adoptions++;
+							}
+							else
+							{
+								NumOrphans++;
+							}
+						} // end while(Orphans)
+						NumOrphans += Orphans.Num();
+
+						//ELogVerbosity::Type Verbosity = Orphans.Num() > 0 ? ELogVerbosity::Error : ELogVerbosity::Display;
+						UE_LOG(LogMeshBindings, Display,
+							TEXT("'%s' - Generated mesh bindings between tet mesh index %d and %s mesh of '%s' LOD %d - stats:\n"
+								"    Render vertices num: %d\n"
+								"    Vertices in tetrahedra: %d\n"
+								"    Vertices bound to tet surface: %d\n"
+								"    Orphaned vertices reparented: %d\n"
+								"    Vertices orphaned: %d"),
+							*GetName().ToString(),
+							TetMeshIdx,
+							bUseSkeletalMeshImportModel ? TEXT("import") : TEXT("render"),
+							*MeshId, LOD,
+							MeshVertices[LOD].Num(), TetHits, TriHits, Adoptions, NumOrphans);
+
+					} // end for all LOD
+
+					// Stash bindings in the geometry collection
+					GeometryCollection::Facades::FTetrahedralBindings TetBindings(OutCollection);
+					TetBindings.DefineSchema();
+					FName MeshName(*MeshId, MeshId.Len());
+					for (int32 LOD = 0; LOD < MeshVertices.Num(); LOD++)
+					{
+						TetBindings.AddBindingsGroup(TetMeshIdx, MeshName, LOD);
+						TetBindings.SetBindingsData(Parents[LOD], Weights[LOD], Offsets[LOD], Masks[LOD]);
+					}
+				} // TetMeshCount > 0
 			} // end for TetMeshIdx
 		}
 		SetValue(Context, MoveTemp(OutCollection), &Collection);
diff --git a/Engine/Source/Runtime/Experimental/Chaos/Private/GeometryCollection/Facades/CollectionTetrahedralBindingsFacade.cpp b/Engine/Source/Runtime/Experimental/Chaos/Private/GeometryCollection/Facades/CollectionTetrahedralBindingsFacade.cpp
--- a/Engine/Source/Runtime/Experimental/Chaos/Private/GeometryCollection/Facades/CollectionTetrahedralBindingsFacade.cpp
+++ b/Engine/Source/Runtime/Experimental/Chaos/Private/GeometryCollection/Facades/CollectionTetrahedralBindingsFacade.cpp
@@ -176,7 +176,7 @@ namespace GeometryCollection::Facades
 			MeshIdAttribute.IsValid() ? MeshIdAttribute.Find() : nullptr;
 		if (MeshIdValues)
 		{
-			FString Suffix = FString::Printf(TEXT(":%s:%d"), *MeshId.ToString(), LOD);
+			FString Suffix = FString::Printf(TEXT(":%s:%d"), *MeshId.GetPlainNameString(), LOD);
 			for (int32 i = 0; i < MeshIdValues->Num(); i++)
 			{
 				const FString& Entry = (*MeshIdValues)[i];
