Fix non-unity build

diff --git a/Engine/Plugins/Experimental/ChaosFlesh/Source/ChaosFleshEngine/Private/ChaosFlesh/ChaosDeformableCollisionsComponent.cpp b/Engine/Plugins/Experimental/ChaosFlesh/Source/ChaosFleshEngine/Private/ChaosFlesh/ChaosDeformableCollisionsComponent.cpp
--- a/Engine/Plugins/Experimental/ChaosFlesh/Source/ChaosFleshEngine/Private/ChaosFlesh/ChaosDeformableCollisionsComponent.cpp
+++ b/Engine/Plugins/Experimental/ChaosFlesh/Source/ChaosFleshEngine/Private/ChaosFlesh/ChaosDeformableCollisionsComponent.cpp
@@ -11,6 +11,8 @@
 #include "Chaos/TriangleMeshImplicitObject.h"
 #include "Components/StaticMeshComponent.h"
 #include "PhysicsEngine/BodySetup.h"
+#include "Engine/StaticMesh.h"
+
 DEFINE_LOG_CATEGORY_STATIC(LogUDeformableCollisionsComponentInternal, Log, All);
 
 #define PERF_SCOPE(X) SCOPE_CYCLE_COUNTER(X); TRACE_CPUPROFILER_EVENT_SCOPE(X);
diff --git a/Engine/Plugins/Experimental/ChaosFlesh/Source/ChaosFleshNodes/Private/Dataflow/ChaosFleshComputeIslandsNode.cpp b/Engine/Plugins/Experimental/ChaosFlesh/Source/ChaosFleshNodes/Private/Dataflow/ChaosFleshComputeIslandsNode.cpp
--- a/Engine/Plugins/Experimental/ChaosFlesh/Source/ChaosFleshNodes/Private/Dataflow/ChaosFleshComputeIslandsNode.cpp
+++ b/Engine/Plugins/Experimental/ChaosFlesh/Source/ChaosFleshNodes/Private/Dataflow/ChaosFleshComputeIslandsNode.cpp
@@ -2,6 +2,7 @@
 
 #include "Dataflow/ChaosFleshComputeIslandsNode.h"
 
+#include "Chaos/Utilities.h"
 #include "ChaosFlesh/TetrahedralCollection.h"
 
 #include UE_INLINE_GENERATED_CPP_BY_NAME(ChaosFleshComputeIslandsNode)
diff --git a/Engine/Plugins/Experimental/ChaosFlesh/Source/ChaosFleshNodes/Private/Dataflow/ChaosFleshSetFleshBonePositionTargetBindingNode.cpp b/Engine/Plugins/Experimental/ChaosFlesh/Source/ChaosFleshNodes/Private/Dataflow/ChaosFleshSetFleshBonePositionTargetBindingNode.cpp
--- a/Engine/Plugins/Experimental/ChaosFlesh/Source/ChaosFleshNodes/Private/Dataflow/ChaosFleshSetFleshBonePositionTargetBindingNode.cpp
+++ b/Engine/Plugins/Experimental/ChaosFlesh/Source/ChaosFleshNodes/Private/Dataflow/ChaosFleshSetFleshBonePositionTargetBindingNode.cpp
@@ -4,6 +4,7 @@
 
 #include "Engine/SkeletalMesh.h"
 #include "Chaos/BoundingVolumeHierarchy.h"
+#include "Chaos/Utilities.h"
 #include "GeometryCollection/Facades/CollectionKinematicBindingFacade.h"
 #include "GeometryCollection/Facades/CollectionPositionTargetFacade.h"
 #include "GeometryCollection/Facades/CollectionVertexBoneWeightsFacade.h"
diff --git a/Engine/Plugins/Experimental/ChaosFlesh/Source/ChaosFleshNodes/Private/Dataflow/ChaosFleshSetVertexTrianglePositionTargetBindingNode.cpp b/Engine/Plugins/Experimental/ChaosFlesh/Source/ChaosFleshNodes/Private/Dataflow/ChaosFleshSetVertexTrianglePositionTargetBindingNode.cpp
--- a/Engine/Plugins/Experimental/ChaosFlesh/Source/ChaosFleshNodes/Private/Dataflow/ChaosFleshSetVertexTrianglePositionTargetBindingNode.cpp
+++ b/Engine/Plugins/Experimental/ChaosFlesh/Source/ChaosFleshNodes/Private/Dataflow/ChaosFleshSetVertexTrianglePositionTargetBindingNode.cpp
@@ -2,6 +2,7 @@
 
 #include "Dataflow/ChaosFleshSetVertexTrianglePositionTargetBindingNode.h"
 
+#include "Chaos/Utilities.h"
 #include "Chaos/HierarchicalSpatialHash.h"
 #include "Chaos/TriangleCollisionPoint.h"
 #include "GeometryCollection/Facades/CollectionMeshFacade.h"
