ChaosFlesh: Transform fix for simulation in PIE

diff --git a/Engine/Plugins/Experimental/ChaosFlesh/Source/ChaosFleshEngine/Private/ChaosFlesh/ChaosDeformableTetrahedralComponent.cpp b/Engine/Plugins/Experimental/ChaosFlesh/Source/ChaosFleshEngine/Private/ChaosFlesh/ChaosDeformableTetrahedralComponent.cpp
--- a/Engine/Plugins/Experimental/ChaosFlesh/Source/ChaosFleshEngine/Private/ChaosFlesh/ChaosDeformableTetrahedralComponent.cpp
+++ b/Engine/Plugins/Experimental/ChaosFlesh/Source/ChaosFleshEngine/Private/ChaosFlesh/ChaosDeformableTetrahedralComponent.cpp
@@ -18,6 +18,7 @@
 #include "Engine/SkeletalMesh.h"
 #include "GeometryCollection/Facades/CollectionTransformSourceFacade.h"
 #include "GeometryCollection/Facades/CollectionTetrahedralSkeletalBindingsFacade.h"
+#include "GeometryCollection/GeometryCollectionAlgo.h"
 #include "GeometryCollection/TransformCollection.h"
 #include "ProceduralMeshComponent.h"
 #include "Rendering/SkeletalMeshRenderData.h"
@@ -330,10 +331,22 @@ void UDeformableTetrahedralComponent::UpdateFromSimulation(const FDataMapValue*
 			TManagedArray<FVector3f>& DynamicVertex = GetDynamicCollection()->GetPositions();
 			const TManagedArray<FVector3f>& SimulationVertex = FleshBuffer->Dynamic.GetAttribute<FVector3f>("Vertex", FGeometryCollection::VerticesGroup);
 
+			TArray<FTransform3f> Transforms;
+			const TManagedArray<int32>& Parent = GetRestCollection()->GetCollection()->GetAttribute<int32>(FTransformCollection::ParentAttribute, FTransformCollection::TransformGroup);
+			const TManagedArray<FTransform3f>& Transform = GetRestCollection()->GetCollection()->GetAttribute<FTransform3f>(FTransformCollection::TransformAttribute, FTransformCollection::TransformGroup);
+			const TManagedArray<int32>& BoneIndex = GetRestCollection()->GetCollection()->GetAttribute<int32>("BoneMap", FGeometryCollection::VerticesGroup);
+			GeometryCollectionAlgo::GlobalMatrices(Transform, Parent, Transforms);
 			// Simulator produces results in component space.
 			for (int i = DynamicVertex.Num() - 1; i >= 0; i--)
 			{
-				DynamicVertex[i] = SimulationVertex[i];
+				if (0 < BoneIndex[i] && BoneIndex[i] < Transforms.Num())
+				{
+					DynamicVertex[i] = Transforms[BoneIndex[i]].TransformPosition(SimulationVertex[i]);
+				}
+				else
+				{
+					DynamicVertex[i] = SimulationVertex[i];
+				}
 			}
 			
 			// p.Chaos.Deformable.FleshDeformer.UpdateGPUBuffersOnTick 1 (default) or 0
@@ -520,13 +533,14 @@ void UDeformableTetrahedralComponent::RenderProceduralMesh()
 							}
 							else
 							{
-								TArray<FVector3f> RenderVertex; 
+								TArray<FVector3f> RenderVertex;
 								Flesh.ComponentSpaceVertices(RenderVertex);
 
 								if (GetDynamicCollection())
 								{
 									const TManagedArray<FVector3f>& DynamicVertex = GetDynamicCollection()->GetPositions();
-									if ((DynamicVertex.Num() > 0) && (DynamicVertex.Num() == RenderVertex.Num())) RenderVertex = DynamicVertex.GetConstArray();
+									if ((DynamicVertex.Num() > 0) && (DynamicVertex.Num() == RenderVertex.Num())) 
+										RenderVertex = DynamicVertex.GetConstArray();
 								}
 								auto InRange = [](int32 Size, int32 Val) { return 0 <= Val && Val < Size; };
 
