ChaosFlesh: Modified Append() in ManagedArrayCollection and derived collections

diff --git a/Engine/Plugins/Experimental/ChaosFlesh/Source/ChaosFlesh/Private/ChaosFlesh/TetrahedralCollection.cpp b/Engine/Plugins/Experimental/ChaosFlesh/Source/ChaosFlesh/Private/ChaosFlesh/TetrahedralCollection.cpp
--- a/Engine/Plugins/Experimental/ChaosFlesh/Source/ChaosFlesh/Private/ChaosFlesh/TetrahedralCollection.cpp
+++ b/Engine/Plugins/Experimental/ChaosFlesh/Source/ChaosFlesh/Private/ChaosFlesh/TetrahedralCollection.cpp
@@ -37,7 +37,7 @@ void FTetrahedralCollection::Construct()
 
 	// Geometry Group
 	AddExternalAttribute<int32>(FTetrahedralCollection::TetrahedronStartAttribute, FGeometryCollection::GeometryGroup, TetrahedronStart, TetrahedronDependency);
-	AddExternalAttribute<int32>(FTetrahedralCollection::TetrahedronCountAttribute, FGeometryCollection::GeometryGroup, TetrahedronCount, TetrahedronDependency);
+	AddExternalAttribute<int32>(FTetrahedralCollection::TetrahedronCountAttribute, FGeometryCollection::GeometryGroup, TetrahedronCount);
 	AddExternalAttribute<FString>(FTetrahedralCollection::GuidAttribute, FGeometryCollection::GeometryGroup, Guid);
 	for (FString& g : Guid) { g = FGuid::NewGuid().ToString(); }
 }
@@ -137,6 +137,29 @@ void FTetrahedralCollection::UpdateBoundingBox()
 	}
 }
 
+void FTetrahedralCollection::Append(const FManagedArrayCollection& InCollection)
+{
+	Super::Append(InCollection);
+	// InCollection data is appended to the front
+	const int32 Offset = InCollection.NumElements(TetrahedralGroup);
+	const int32 OtherSize = InCollection.NumElements(VerticesGroup);
+	const int32 Size = NumElements(VerticesGroup);
+	// IncidentElements does not have TetrahedralGroup dependency, update manually
+	for (int32 Idx = OtherSize; Idx < Size; Idx++)
+	{
+		for (int32& TetIdx: IncidentElements[Idx])
+		{
+			// Offset by the number of tets of the second collection
+			TetIdx += Offset;
+		}
+	}
+}
+
+void FTetrahedralCollection::AppendCollection(const FTetrahedralCollection& InCollection)
+{
+	Append(InCollection);
+}
+
 int32 FTetrahedralCollection::AppendGeometry(
 	const FTetrahedralCollection& Other, 
 	int32 MaterialIDOffset, 
@@ -277,4 +300,4 @@ void FTetrahedralCollection::ReorderTetrahedralElements(const TArray<int32>& New
 	}
 
 	Super::ReorderElements(TetrahedralGroup, NewTetOrder);
-}
+}
\ No newline at end of file
diff --git a/Engine/Plugins/Experimental/ChaosFlesh/Source/ChaosFlesh/Public/ChaosFlesh/TetrahedralCollection.h b/Engine/Plugins/Experimental/ChaosFlesh/Source/ChaosFlesh/Public/ChaosFlesh/TetrahedralCollection.h
--- a/Engine/Plugins/Experimental/ChaosFlesh/Source/ChaosFlesh/Public/ChaosFlesh/TetrahedralCollection.h
+++ b/Engine/Plugins/Experimental/ChaosFlesh/Source/ChaosFlesh/Public/ChaosFlesh/TetrahedralCollection.h
@@ -16,12 +16,12 @@ namespace Chaos
 /**
 * FTetrahedralCollection (FGeometryCollection)
 */
-class CHAOSFLESH_API FTetrahedralCollection : public FGeometryCollection
+class FTetrahedralCollection : public FGeometryCollection
 {
 	typedef FGeometryCollection Super;
 
 public:
-	FTetrahedralCollection();
+	CHAOSFLESH_API FTetrahedralCollection();
 	FTetrahedralCollection(FTetrahedralCollection &) = delete;
 	FTetrahedralCollection& operator=(const FTetrahedralCollection &) = delete;
 	FTetrahedralCollection(FTetrahedralCollection &&) = default;
@@ -30,73 +30,76 @@ public:
 	/**
 	 * Create a GeometryCollection from Vertex and Indices arrays
 	 */
-	static FTetrahedralCollection* NewTetrahedralCollection(const TArray<FVector>& Vertices, const TArray<FIntVector3>& SurfaceElements, const TArray<FIntVector4>& Elements,bool bReverseVertexOrder = true);
-	static void Init(FTetrahedralCollection* Collection, const TArray<FVector>& Vertices, const TArray<FIntVector3>& SurfaceElements, const TArray<FIntVector4>& Elements, bool bReverseVertexOrder = true);
+	CHAOSFLESH_API static FTetrahedralCollection* NewTetrahedralCollection(const TArray<FVector>& Vertices, const TArray<FIntVector3>& SurfaceElements, const TArray<FIntVector4>& Elements,bool bReverseVertexOrder = true);
+	CHAOSFLESH_API static void Init(FTetrahedralCollection* Collection, const TArray<FVector>& Vertices, const TArray<FIntVector3>& SurfaceElements, const TArray<FIntVector4>& Elements, bool bReverseVertexOrder = true);
 
-	void UpdateBoundingBox();
+	CHAOSFLESH_API void UpdateBoundingBox();
 
-	int32 AppendGeometry(const FTetrahedralCollection& GeometryCollection, int32 MaterialIDOffset = 0, bool ReindexAllMaterials = true, const FTransform& TransformRoot = FTransform::Identity);
+	CHAOSFLESH_API void AppendCollection(const FTetrahedralCollection& InCollection);
+
+	CHAOSFLESH_API int32 AppendGeometry(const FTetrahedralCollection& GeometryCollection, int32 MaterialIDOffset = 0, bool ReindexAllMaterials = true, const FTransform& TransformRoot = FTransform::Identity);
 
 	/**
 	* Build \c IncidentElements and \c IncidentElementsLocalIndex.
 	* \p GeometryIndex - ID of the geometry entry to restrict operation to; -1 does all.
 	*/
-	void InitIncidentElements(const int32 GeometryIndex=-1);
+	CHAOSFLESH_API void InitIncidentElements(const int32 GeometryIndex=-1);
 
 	/*
 	*  SetDefaults for new entries on this collection. 
 	*/
-	virtual void SetDefaults(FName Group, uint32 StartSize, uint32 NumElements) override;
+	CHAOSFLESH_API virtual void SetDefaults(FName Group, uint32 StartSize, uint32 NumElements) override;
 
 	/**
 	* Reorders elements in a group. NewOrder must be the same length as the group.
 	*/
-	void ReorderElements(FName Group, const TArray<int32>& NewOrder) override;
+	CHAOSFLESH_API void ReorderElements(FName Group, const TArray<int32>& NewOrder) override;
 
-	void ReorderTetrahedralElements(const TArray<int32>& NewOrder);
+	CHAOSFLESH_API void ReorderTetrahedralElements(const TArray<int32>& NewOrder);
 
 	/*
 	*  Attribute Groups
 	*/
-	static const FName TetrahedralGroup;
-	static const FName BindingsGroup;
+	CHAOSFLESH_API static const FName TetrahedralGroup;
+	CHAOSFLESH_API static const FName BindingsGroup;
 
 	/*
 	*  Tetrahedron Attribute
 	*  TManagedArray<FIntVector4> Tetrahedron = this->FindAttribute<FIntVector4>(FTetrahedralCollection::TetrahedronAttribute,FTetrahedralCollection::TetrahedralGroup);
 	*/
-	static const FName TetrahedronAttribute;
+	CHAOSFLESH_API static const FName TetrahedronAttribute;
 	TManagedArray<FIntVector4> Tetrahedron;
 
-	static const FName TetrahedronStartAttribute;
+	CHAOSFLESH_API static const FName TetrahedronStartAttribute;
 	TManagedArray<int32> TetrahedronStart;
-	static const FName TetrahedronCountAttribute;
+	CHAOSFLESH_API static const FName TetrahedronCountAttribute;
 	TManagedArray<int32> TetrahedronCount;
 
 	/**
 	* Incident Elements Attribute
 	* For each vertex, a list of tetrahedra that includes that vertex.
 	*/
-	static const FName IncidentElementsAttribute;
+	CHAOSFLESH_API static const FName IncidentElementsAttribute;
 	TManagedArray<TArray<int32>> IncidentElements;
 
 	/**
 	* Incident Elements Attribute
 	* For each incident element, the vertex's index in the tetrahedron.
 	*/
-	static const FName IncidentElementsLocalIndexAttribute;
+	CHAOSFLESH_API static const FName IncidentElementsLocalIndexAttribute;
 	TManagedArray<TArray<int32>> IncidentElementsLocalIndex;
 
 	/*
 	* Guid Attribute
 	*/
-	static const FName GuidAttribute;
+	CHAOSFLESH_API static const FName GuidAttribute;
 	TManagedArray<FString> Guid;
 
 protected:
 
-	void Construct();
+	CHAOSFLESH_API void Construct();
 
+	CHAOSFLESH_API virtual void Append(const FManagedArrayCollection& InCollection) override;
 };
 
 FORCEINLINE Chaos::FChaosArchive& operator<<(Chaos::FChaosArchive& Ar, FTetrahedralCollection& Value)
diff --git a/Engine/Plugins/Experimental/ChaosFlesh/Source/ChaosFleshNodes/Private/Dataflow/ChaosFleshAppendTetrahedralCollectionNode.cpp b/Engine/Plugins/Experimental/ChaosFlesh/Source/ChaosFleshNodes/Private/Dataflow/ChaosFleshAppendTetrahedralCollectionNode.cpp
--- a/Engine/Plugins/Experimental/ChaosFlesh/Source/ChaosFleshNodes/Private/Dataflow/ChaosFleshAppendTetrahedralCollectionNode.cpp
+++ b/Engine/Plugins/Experimental/ChaosFlesh/Source/ChaosFleshNodes/Private/Dataflow/ChaosFleshAppendTetrahedralCollectionNode.cpp
@@ -18,7 +18,7 @@ void FAppendTetrahedralCollectionDataflowNode::Evaluate(Dataflow::FContext& Cont
 		}
 		if (InCollection2)
 		{
-			InCollection1->AppendGeometry(*InCollection2);
+			InCollection1->AppendCollection(*InCollection2);
 		}		
 		if (const TManagedArray<FString>* GuidArray2 = InCollection2->FindAttribute<FString>("Guid", FGeometryCollection::GeometryGroup))
 		{
diff --git a/Engine/Source/Runtime/Experimental/Chaos/Private/GeometryCollection/Facades/CollectionKinematicBindingFacade.cpp b/Engine/Source/Runtime/Experimental/Chaos/Private/GeometryCollection/Facades/CollectionKinematicBindingFacade.cpp
--- a/Engine/Source/Runtime/Experimental/Chaos/Private/GeometryCollection/Facades/CollectionKinematicBindingFacade.cpp
+++ b/Engine/Source/Runtime/Experimental/Chaos/Private/GeometryCollection/Facades/CollectionKinematicBindingFacade.cpp
@@ -45,7 +45,7 @@ namespace GeometryCollection::Facades
 	FKinematicBindingFacade::SetBoneBindings(const int32 InBoneIndex, const TArray<int32>& InBoneVerts, const TArray<float>& InBoneWeights)
 	{
 		check(!IsConst());
-		return GeometryCollection::Facades::FSelectionFacade(*Collection).AddSelection(InBoneIndex, InBoneVerts, InBoneWeights, FGeometryCollection::VerticesGroup);
+		return GeometryCollection::Facades::FSelectionFacade(*Collection).AddSelection(InBoneIndex, InBoneVerts, InBoneWeights, FGeometryCollection::VerticesGroup, FTransformCollection::TransformGroup);
 	}
 
 	void FKinematicBindingFacade::GetBoneBindings(const FKinematicBindingFacade::FBindingKey& Key, int32& OutBoneIndex, TArray<int32>& OutBoneVerts, TArray<float>& OutBoneWeights) const
diff --git a/Engine/Source/Runtime/Experimental/Chaos/Private/GeometryCollection/GeometryCollection.cpp b/Engine/Source/Runtime/Experimental/Chaos/Private/GeometryCollection/GeometryCollection.cpp
--- a/Engine/Source/Runtime/Experimental/Chaos/Private/GeometryCollection/GeometryCollection.cpp
+++ b/Engine/Source/Runtime/Experimental/Chaos/Private/GeometryCollection/GeometryCollection.cpp
@@ -178,6 +178,31 @@ void FGeometryCollection::SetDefaults(FName Group, uint32 StartSize, uint32 NumE
 	}
 }
 
+void FGeometryCollection::Append(const FManagedArrayCollection& InCollection)
+{	
+	/* InCollection data is appended to the front*/
+	Super::Append(InCollection);
+	if (const FGeometryCollection* InTypedCollection = InCollection.Cast<FGeometryCollection>())
+	{
+		const int32 Offset = InTypedCollection->NumElements(GeometryGroup);
+		const int32 OtherSize = InTypedCollection->NumElements(TransformGroup);
+		const int32 Size = NumElements(TransformGroup);
+		/*TransformToGeometryIndex does not have GeometryGroup dependency, update manually*/
+		for (int32 Idx = OtherSize; Idx < Size; ++Idx)
+		{
+			if (TransformToGeometryIndex[Idx] != INDEX_NONE)
+			{
+				TransformToGeometryIndex[Idx] += Offset;
+			}
+		}
+	}
+}
+
+void FGeometryCollection::AppendCollection(const FGeometryCollection& InCollection)
+{
+	Append(InCollection);
+}
+
 // MaterialIDOffset is based on the number of materials added by this append geometry call
 int32 FGeometryCollection::AppendGeometry(const FGeometryCollection & Element, int32 MaterialIDOffset, bool ReindexAllMaterials, const FTransform& TransformRoot)
 {
@@ -757,7 +782,10 @@ void FGeometryCollection::ReorderTransformElements(const TArray<int32>& NewOrder
 	Pairs.Reserve(NumGeometries);
 	for (int32 GeomIdx = 0; GeomIdx < NumGeometries; ++GeomIdx)
 	{
-		Pairs.Emplace(NewOrder[TransformIndex[GeomIdx]], GeomIdx);
+		if (TransformIndex[GeomIdx] != INDEX_NONE)
+		{
+			Pairs.Emplace(NewOrder[TransformIndex[GeomIdx]], GeomIdx);
+		}
 	}
 	Pairs.Sort();
 
diff --git a/Engine/Source/Runtime/Experimental/Chaos/Private/GeometryCollection/ManagedArrayCollection.cpp b/Engine/Source/Runtime/Experimental/Chaos/Private/GeometryCollection/ManagedArrayCollection.cpp
--- a/Engine/Source/Runtime/Experimental/Chaos/Private/GeometryCollection/ManagedArrayCollection.cpp
+++ b/Engine/Source/Runtime/Experimental/Chaos/Private/GeometryCollection/ManagedArrayCollection.cpp
@@ -232,6 +232,32 @@ int32 FManagedArrayCollection::InsertElements(int32 NumberElements, int32 Positi
 	return Position;
 }
 
+TArray<int32> FManagedArrayCollection::InsertElementsNoReorder(int32 NumberElements, int32 Position, FName Group)
+{
+	const int32 OldGroupSize = AddElements(NumberElements, Group);
+	const int32 NewGroupSize = OldGroupSize + NumberElements;
+	check(Position <= OldGroupSize);
+	const int32 NumberElementsToMove = OldGroupSize - Position;
+	const int32 MoveToPosition = Position + NumberElements;
+
+	TArray<int32> NewOrder;
+	NewOrder.SetNumUninitialized(NewGroupSize);
+
+	for (int32 Idx = 0; Idx < Position; ++Idx)
+	{
+		NewOrder[Idx] = Idx;
+	}
+	for (int32 Idx = Position; Idx < MoveToPosition; ++Idx)
+	{
+		NewOrder[Idx] = Idx + NumberElementsToMove;
+	}
+	for (int32 Idx = MoveToPosition; Idx < NewGroupSize; ++Idx)
+	{
+		NewOrder[Idx] = Idx - NumberElements;
+	}
+	return NewOrder;
+}
+
 void FManagedArrayCollection::Append(const FManagedArrayCollection& InCollection)
 {
 	bool bMatchingAttributes = true;
@@ -253,12 +279,13 @@ void FManagedArrayCollection::Append(const FManagedArrayCollection& InCollection
 	}
 	if (bMatchingAttributes)
 	{
+		TMap<FName, TArray<int32>> GroupNewOrder;
 		// make space first. 
 		for (const FName& Group : InCollection.GroupNames())
 		{
-			if (HasGroup(Group) && NumElements(Group))
+			if (HasGroup(Group) && InCollection.NumElements(Group))
 			{
-				InsertElements(InCollection.NumElements(Group), 0, Group);
+				GroupNewOrder.Add(Group, InsertElementsNoReorder(InCollection.NumElements(Group), 0, Group));
 			}
 			else if (!HasGroup(Group))
 			{
@@ -266,6 +293,10 @@ void FManagedArrayCollection::Append(const FManagedArrayCollection& InCollection
 				AddElements(InCollection.NumElements(Group), Group);
 			}
 		}
+		for (TTuple<FName, TArray<int32>>& Pair : GroupNewOrder)
+		{
+			FManagedArrayCollection::ReorderElements(Pair.Key, Pair.Value);
+		}
 
 		// copy values
 		for (const TTuple<FKeyType, FValueType>& Entry : InCollection.Map)
@@ -280,6 +311,7 @@ void FManagedArrayCollection::Append(const FManagedArrayCollection& InCollection
 			else
 			{
 				FValueType NewAttribute(Entry.Value);
+				NewAttribute.Resize(NumElements(GroupName));
 				Map.Add(Entry.Key, MoveTemp(NewAttribute));
 			}
 		}
diff --git a/Engine/Source/Runtime/Experimental/Chaos/Public/GeometryCollection/GeometryCollection.h b/Engine/Source/Runtime/Experimental/Chaos/Public/GeometryCollection/GeometryCollection.h
--- a/Engine/Source/Runtime/Experimental/Chaos/Public/GeometryCollection/GeometryCollection.h
+++ b/Engine/Source/Runtime/Experimental/Chaos/Public/GeometryCollection/GeometryCollection.h
@@ -152,10 +152,10 @@ public:
 		const TManagedArray<int32>& RawStatusFlagsArray,
 		FGeometryCollectionDefaults InDefaults = FGeometryCollectionDefaults());
 
-	//
-	//
-	//
-
+	/**
+	* Append a FGeometryCollection to a FGeometryCollection
+	*/
+	CHAOS_API void AppendCollection(const FGeometryCollection& InCollection);
 
 	/** 
 	* Append a single geometric object to a FGeometryCollection 
@@ -388,6 +388,7 @@ protected:
 		return 10;
 	}
 
+	CHAOS_API virtual void Append(const FManagedArrayCollection& InCollection) override;
 
 public:
 	/* Backwards compatibility */
diff --git a/Engine/Source/Runtime/Experimental/Chaos/Public/GeometryCollection/ManagedArray.h b/Engine/Source/Runtime/Experimental/Chaos/Public/GeometryCollection/ManagedArray.h
--- a/Engine/Source/Runtime/Experimental/Chaos/Public/GeometryCollection/ManagedArray.h
+++ b/Engine/Source/Runtime/Experimental/Chaos/Public/GeometryCollection/ManagedArray.h
@@ -805,7 +805,7 @@ public:
 	virtual void CopyRange(const FManagedArrayBase& ConstArray, int32 Start, int32 Stop, int32 Offset = 0) override
 	{
 		ensureMsgf(ConstArray.GetTypeSize() == GetTypeSize(), TEXT("TManagedArrayBase<T>::Init : Invalid array types."));
-		if (ensureMsgf(Stop + Offset < Array.Num(), TEXT("Error : Index out of bounds")))
+		if (ensureMsgf(Stop + Offset <= Array.Num(), TEXT("Error : Index out of bounds")))
 		{
 			const FManagedBitArrayBase& TypedConstArray = static_cast<const FManagedBitArrayBase&>(ConstArray);
 			for (int32 Sdx = Start, Tdx = Start + Offset; Sdx < ConstArray.Num() && Tdx < Array.Num() && Sdx < Stop; Sdx++, Tdx++)
diff --git a/Engine/Source/Runtime/Experimental/Chaos/Public/GeometryCollection/ManagedArrayCollection.h b/Engine/Source/Runtime/Experimental/Chaos/Public/GeometryCollection/ManagedArrayCollection.h
--- a/Engine/Source/Runtime/Experimental/Chaos/Public/GeometryCollection/ManagedArrayCollection.h
+++ b/Engine/Source/Runtime/Experimental/Chaos/Public/GeometryCollection/ManagedArrayCollection.h
@@ -272,7 +272,6 @@ public:
 	*/
 	CHAOS_API int32 InsertElements(int32 NumberElements, int32 Position, FName Group);
 
-
 	/**
 	* Append Collection and reindex dependencies on this collection. 
 	* @param InCollection : Collection to add. 
@@ -618,6 +617,15 @@ private:
 		Map.Add({ Name, Group }, MoveTemp(NewAttribute));
 	}
 
+	/**
+	* Insert elements to a group without reordering elements
+	* @param NumberElements - The number of array entries to add
+	* @param Position - The position in the managed array where to insert entries.
+	* @param Group - The group to append entries to.
+	* @return new order after insertion.
+	*/
+	TArray<int32> InsertElementsNoReorder(int32 NumberElements, int32 Position, FName Group);
+
 	/****
 	*  Mapping Key/Value
 	*
diff --git a/Engine/Source/Runtime/Experimental/Chaos/Public/GeometryCollection/TransformCollection.h b/Engine/Source/Runtime/Experimental/Chaos/Public/GeometryCollection/TransformCollection.h
--- a/Engine/Source/Runtime/Experimental/Chaos/Public/GeometryCollection/TransformCollection.h
+++ b/Engine/Source/Runtime/Experimental/Chaos/Public/GeometryCollection/TransformCollection.h
@@ -70,9 +70,15 @@ public:
 
 	/**
 	* Append Collection and reindex dependencies on this collection.
-	* @param InCollection : Collection to add.
+	* @param InCollection : FManagedArrayCollection to add.
 	*/
-	CHAOS_API virtual void Append(const FManagedArrayCollection& Collection);
+	CHAOS_API virtual void Append(const FManagedArrayCollection& Collection) override;
+
+	/**
+	* Append Collection and reindex dependencies on this collection.
+	* @param InCollection : FTransformCollection to add.
+	*/
+	CHAOS_API void AppendCollection(const FTransformCollection& Collection);
 
 	/*
 	* AppendTransform:
