ChaosFlesh: fix for CIS
move FleshCollectionUtility to FleshCollectionEngineUtility
FContextSingle constructor takes no argument now in dataflow test

diff --git a/Engine/Plugins/Animation/MLDeformer/ChaosFleshGenerator/Source/ChaosFleshGenerator/Private/ChaosFleshGeneratorPrivate.cpp b/Engine/Plugins/Animation/MLDeformer/ChaosFleshGenerator/Source/ChaosFleshGenerator/Private/ChaosFleshGeneratorPrivate.cpp
--- a/Engine/Plugins/Animation/MLDeformer/ChaosFleshGenerator/Source/ChaosFleshGenerator/Private/ChaosFleshGeneratorPrivate.cpp
+++ b/Engine/Plugins/Animation/MLDeformer/ChaosFleshGenerator/Source/ChaosFleshGenerator/Private/ChaosFleshGeneratorPrivate.cpp
@@ -3,7 +3,7 @@
 #include "ChaosFleshGeneratorPrivate.h"
 
 #include "Chaos/Vector.h"
-#include "ChaosFlesh/FleshCollectionUtility.h"
+#include "ChaosFlesh/FleshCollectionEngineUtility.h"
 #include "FleshGeneratorComponent.h"
 #include "FileHelpers.h"
 #include "GeometryCollection/Facades/CollectionTetrahedralBindingsFacade.h"
diff --git a/Engine/Plugins/Animation/MLDeformer/ChaosFleshGenerator/Source/ChaosFleshGenerator/Private/ChaosFleshGeneratorSimulation.cpp b/Engine/Plugins/Animation/MLDeformer/ChaosFleshGenerator/Source/ChaosFleshGenerator/Private/ChaosFleshGeneratorSimulation.cpp
--- a/Engine/Plugins/Animation/MLDeformer/ChaosFleshGenerator/Source/ChaosFleshGenerator/Private/ChaosFleshGeneratorSimulation.cpp
+++ b/Engine/Plugins/Animation/MLDeformer/ChaosFleshGenerator/Source/ChaosFleshGenerator/Private/ChaosFleshGeneratorSimulation.cpp
@@ -11,7 +11,7 @@
 #include "FleshGeneratorProperties.h"
 #include "Engine/SkeletalMesh.h"
 #include "Logging/LogMacros.h"
-#include "ChaosFlesh/FleshCollectionUtility.h"
+#include "ChaosFlesh/FleshCollectionEngineUtility.h"
 
 DEFINE_LOG_CATEGORY(LogChaosFleshGeneratorSimulation);
 
diff --git a/Engine/Plugins/Experimental/ChaosFlesh/Source/ChaosFlesh/ChaosFlesh.Build.cs b/Engine/Plugins/Experimental/ChaosFlesh/Source/ChaosFlesh/ChaosFlesh.Build.cs
--- a/Engine/Plugins/Experimental/ChaosFlesh/Source/ChaosFlesh/ChaosFlesh.Build.cs
+++ b/Engine/Plugins/Experimental/ChaosFlesh/Source/ChaosFlesh/ChaosFlesh.Build.cs
@@ -10,7 +10,6 @@ namespace UnrealBuildTool.Rules
 			PublicDependencyModuleNames.AddRange(
 				new string[]
 				{
-					"Engine",
 					"Core",
 					"CoreUObject",
 					"Chaos",
diff --git a/Engine/Plugins/Experimental/ChaosFlesh/Source/ChaosFlesh/Private/ChaosFlesh/FleshCollectionUtility.cpp b/Engine/Plugins/Experimental/ChaosFlesh/Source/ChaosFlesh/Private/ChaosFlesh/FleshCollectionUtility.cpp
--- a/Engine/Plugins/Experimental/ChaosFlesh/Source/ChaosFlesh/Private/ChaosFlesh/FleshCollectionUtility.cpp
+++ b/Engine/Plugins/Experimental/ChaosFlesh/Source/ChaosFlesh/Private/ChaosFlesh/FleshCollectionUtility.cpp
@@ -8,11 +8,8 @@
 
 #include "ChaosFlesh/ChaosFlesh.h"
 #include "ChaosFlesh/FleshCollection.h"
-#include "GeometryCollection/Facades/CollectionTetrahedralBindingsFacade.h"
 #include "GeometryCollection/TransformCollection.h"
 #include "Misc/Paths.h"
-#include "Engine/SkeletalMesh.h"
-#include "Engine/StaticMesh.h"
 
 #include <fstream>
 
@@ -54,7 +51,7 @@ namespace ChaosFlesh
 	}
 	FIntVector4	GetOrdered(const FIntVector4& V)
 	{
-		TArray<int32, TInlineAllocator<4>> VA = { V[0], V[1], V[2] };
+		TArray<int32, TInlineAllocator<4>> VA = { V[0], V[1], V[2], V[3]};
 		VA.Sort();
 		return FIntVector4(VA[0], VA[1], VA[2], VA[3]);
 	}
@@ -155,76 +152,4 @@ namespace ChaosFlesh
 				UE_LOG(LogChaosFlesh, Display, TEXT("WARNING: Encountered %d non-manifold tetrahedral mesh faces."), nonManifold);
 		}
 	}
-
-	FString GetMeshId(const USkeletalMesh* SkeletalMesh, const bool bUseImportModel)
-	{
-		FPrimaryAssetId Id = SkeletalMesh->GetPrimaryAssetId();
-		FString MeshId = Id.IsValid() ? Id.ToString() : SkeletalMesh->GetName();
-		if (bUseImportModel)
-		{
-			MeshId.Append(TEXT("_ImportModel"));
-		}
-		return MeshId;
-	}
-
-	FString GetMeshId(const UStaticMesh* StaticMesh)
-	{
-		FPrimaryAssetId Id = StaticMesh->GetPrimaryAssetId();
-		FString MeshId = Id.IsValid() ? Id.ToString() : StaticMesh->GetName();
-		return MeshId;
-	}
-
-	void BoundSurfacePositions(
-		const USkeletalMesh* SkeletalMesh,
-		const FFleshCollection* FleshCollection,
-		const TManagedArray<FVector3f>* RestVertices,
-		const TManagedArray<FVector3f>* SimulatedVertices,
-		TArray<FVector3f>& Positions)
-	{
-		GeometryCollection::Facades::FTetrahedralBindings TetBindings(*FleshCollection);
-		FString MeshId = GetMeshId(SkeletalMesh, false);
-		FName MeshIdName(MeshId);
-
-		const int32 LODIndex = 0;
-		const int32 TetIndex = TetBindings.GetTetMeshIndex(MeshIdName, LODIndex);
-		if (TetIndex == INDEX_NONE)
-		{
-			UE_LOG(LogChaosFlesh, Error,
-				TEXT("CreateGeometryCache - No tet mesh index associated with mesh '%s' LOD: %d"),
-				*MeshId, LODIndex);
-			return;
-		}
-		if (!TetBindings.ReadBindingsGroup(TetIndex, MeshIdName, LODIndex))
-		{
-			UE_LOG(LogChaosFlesh, Error,
-				TEXT("CreateGeometryCache - Failed to read bindings group associated with mesh '%s' LOD: %d"),
-				*MeshId, LODIndex);
-			return;
-		}
-
-		TUniquePtr<GeometryCollection::Facades::FTetrahedralBindings::Evaluator> BindingsEvalPtr = TetBindings.InitEvaluator(RestVertices);
-		const GeometryCollection::Facades::FTetrahedralBindings::Evaluator& BindingsEval = *BindingsEvalPtr.Get();
-		if (!BindingsEval.IsValid())
-		{
-			UE_LOG(LogChaosFlesh, Error,
-				TEXT("CreateGeometryCache - Invalid flesh bindings for skeletal mesh asset [%s]"),
-				*SkeletalMesh->GetName());
-			return;
-		}
-
-		TArray<Chaos::FVec3f> CurrVertices;
-		CurrVertices.SetNum(SimulatedVertices->Num());
-		for (int32 Index = 0; Index < CurrVertices.Num(); ++Index)
-		{
-			CurrVertices[Index] = Chaos::FVec3f((*SimulatedVertices)[Index]);
-		}
-
-		const int32 NumVertices = BindingsEval.NumVertices();
-		Positions.SetNum(NumVertices);
-		for (int32 Index = 0; Index < NumVertices; ++Index)
-		{
-			Positions[Index] = BindingsEval.GetEmbeddedPosition(Index, CurrVertices);
-		}
-		
-	}
 }
diff --git a/Engine/Plugins/Experimental/ChaosFlesh/Source/ChaosFlesh/Public/ChaosFlesh/FleshCollectionUtility.h b/Engine/Plugins/Experimental/ChaosFlesh/Source/ChaosFlesh/Public/ChaosFlesh/FleshCollectionUtility.h
--- a/Engine/Plugins/Experimental/ChaosFlesh/Source/ChaosFlesh/Public/ChaosFlesh/FleshCollectionUtility.h
+++ b/Engine/Plugins/Experimental/ChaosFlesh/Source/ChaosFlesh/Public/ChaosFlesh/FleshCollectionUtility.h
@@ -7,24 +7,8 @@
 #include "Templates/UniquePtr.h"
 #include "Misc/AssertionMacros.h"
 
-class USkeletalMesh;
-class UFleshAsset;
-class UStaticMesh;
-
 namespace ChaosFlesh 
 {
 	void GetTetFaces(const FIntVector4& Tet, FIntVector3& Face1, FIntVector3& Face2, FIntVector3& Face3, FIntVector3& Face4, const bool invert = false);
-
 	void GetSurfaceElements(const TArray<FIntVector4>& Tets, TArray<FIntVector3>& SurfaceElements, const bool KeepInteriorFaces, const bool InvertFaces = false);
-	
-	CHAOSFLESH_API FString GetMeshId(const USkeletalMesh* SkeletalMesh, const bool bUseImportModel);
-
-	CHAOSFLESH_API FString GetMeshId(const UStaticMesh* StaticMesh);
-
-	CHAOSFLESH_API void BoundSurfacePositions(
-		const USkeletalMesh* SkeletalMesh,
-		const FFleshCollection* FleshCollection,
-		const TManagedArray<FVector3f>* RestVertices,
-		const TManagedArray<FVector3f>* SimulatedVertices,
-		TArray<FVector3f>& Positions);
 }
diff --git a/Engine/Plugins/Experimental/ChaosFlesh/Source/ChaosFleshEditor/Private/ChaosFlesh/Cmd/ChaosFleshCommands.cpp b/Engine/Plugins/Experimental/ChaosFlesh/Source/ChaosFleshEditor/Private/ChaosFlesh/Cmd/ChaosFleshCommands.cpp
--- a/Engine/Plugins/Experimental/ChaosFlesh/Source/ChaosFleshEditor/Private/ChaosFlesh/Cmd/ChaosFleshCommands.cpp
+++ b/Engine/Plugins/Experimental/ChaosFlesh/Source/ChaosFleshEditor/Private/ChaosFlesh/Cmd/ChaosFleshCommands.cpp
@@ -6,7 +6,7 @@
 #include "ChaosFlesh/Cmd/FleshAssetConversion.h"
 #include "ChaosFlesh/FleshAsset.h"
 #include "ChaosFlesh/FleshCollection.h"
-#include "ChaosFlesh/FleshCollectionUtility.h"
+#include "ChaosFlesh/FleshCollectionEngineUtility.h"
 #include "ChaosFlesh/ChaosDeformableTetrahedralComponent.h"
 
 #include "Chaos/CacheManagerActor.h"
diff --git a/Engine/Plugins/Experimental/ChaosFlesh/Source/ChaosFleshEngine/Private/ChaosFlesh/ChaosDeformableTetrahedralComponent.cpp b/Engine/Plugins/Experimental/ChaosFlesh/Source/ChaosFleshEngine/Private/ChaosFlesh/ChaosDeformableTetrahedralComponent.cpp
--- a/Engine/Plugins/Experimental/ChaosFlesh/Source/ChaosFleshEngine/Private/ChaosFlesh/ChaosDeformableTetrahedralComponent.cpp
+++ b/Engine/Plugins/Experimental/ChaosFlesh/Source/ChaosFleshEngine/Private/ChaosFlesh/ChaosDeformableTetrahedralComponent.cpp
@@ -12,7 +12,7 @@
 #include "ChaosFlesh/ChaosDeformableSolverComponent.h"
 #include "ChaosFlesh/ChaosDeformableTypes.h"
 #include "ChaosFlesh/ChaosFleshCollectionFacade.h"
-#include "ChaosFlesh/FleshCollectionUtility.h"
+#include "ChaosFlesh/FleshCollectionEngineUtility.h"
 #include "ChaosFlesh/FleshDynamicAsset.h"
 #include "Components/SkeletalMeshComponent.h"
 #include "Dataflow/DataflowEngineUtil.h"
diff --git a/Engine/Plugins/Experimental/ChaosFlesh/Source/ChaosFleshEngine/Private/ChaosFlesh/FleshCollectionEngineUtility.cpp b/Engine/Plugins/Experimental/ChaosFlesh/Source/ChaosFleshEngine/Private/ChaosFlesh/FleshCollectionEngineUtility.cpp
new file mode 100644
--- /dev/null
+++ b/Engine/Plugins/Experimental/ChaosFlesh/Source/ChaosFleshEngine/Private/ChaosFlesh/FleshCollectionEngineUtility.cpp
@@ -0,0 +1,88 @@
+// Copyright Epic Games, Inc. All Rights Reserved.
+
+/*=============================================================================
+	ChaosDeformableTetrahedralComponent methods.
+=============================================================================*/
+
+#include "ChaosFlesh/FleshCollectionEngineUtility.h"
+
+#include "ChaosFlesh/ChaosFlesh.h"
+#include "ChaosFlesh/FleshCollection.h"
+#include "GeometryCollection/Facades/CollectionTetrahedralBindingsFacade.h"
+#include "Engine/SkeletalMesh.h"
+#include "Engine/StaticMesh.h"
+
+namespace ChaosFlesh
+{
+	FString GetMeshId(const USkeletalMesh* SkeletalMesh, const bool bUseImportModel)
+	{
+		FPrimaryAssetId Id = SkeletalMesh->GetPrimaryAssetId();
+		FString MeshId = Id.IsValid() ? Id.ToString() : SkeletalMesh->GetName();
+		if (bUseImportModel)
+		{
+			MeshId.Append(TEXT("_ImportModel"));
+		}
+		return MeshId;
+	}
+
+	FString GetMeshId(const UStaticMesh* StaticMesh)
+	{
+		FPrimaryAssetId Id = StaticMesh->GetPrimaryAssetId();
+		FString MeshId = Id.IsValid() ? Id.ToString() : StaticMesh->GetName();
+		return MeshId;
+	}
+
+	void BoundSurfacePositions(
+		const USkeletalMesh* SkeletalMesh,
+		const FFleshCollection* FleshCollection,
+		const TManagedArray<FVector3f>* RestVertices,
+		const TManagedArray<FVector3f>* SimulatedVertices,
+		TArray<FVector3f>& Positions)
+	{
+		GeometryCollection::Facades::FTetrahedralBindings TetBindings(*FleshCollection);
+		FString MeshId = GetMeshId(SkeletalMesh, false);
+		FName MeshIdName(MeshId);
+
+		const int32 LODIndex = 0;
+		const int32 TetIndex = TetBindings.GetTetMeshIndex(MeshIdName, LODIndex);
+		if (TetIndex == INDEX_NONE)
+		{
+			UE_LOG(LogChaosFlesh, Error,
+				TEXT("CreateGeometryCache - No tet mesh index associated with mesh '%s' LOD: %d"),
+				*MeshId, LODIndex);
+			return;
+		}
+		if (!TetBindings.ReadBindingsGroup(TetIndex, MeshIdName, LODIndex))
+		{
+			UE_LOG(LogChaosFlesh, Error,
+				TEXT("CreateGeometryCache - Failed to read bindings group associated with mesh '%s' LOD: %d"),
+				*MeshId, LODIndex);
+			return;
+		}
+
+		TUniquePtr<GeometryCollection::Facades::FTetrahedralBindings::Evaluator> BindingsEvalPtr = TetBindings.InitEvaluator(RestVertices);
+		const GeometryCollection::Facades::FTetrahedralBindings::Evaluator& BindingsEval = *BindingsEvalPtr.Get();
+		if (!BindingsEval.IsValid())
+		{
+			UE_LOG(LogChaosFlesh, Error,
+				TEXT("CreateGeometryCache - Invalid flesh bindings for skeletal mesh asset [%s]"),
+				*SkeletalMesh->GetName());
+			return;
+		}
+
+		TArray<Chaos::FVec3f> CurrVertices;
+		CurrVertices.SetNum(SimulatedVertices->Num());
+		for (int32 Index = 0; Index < CurrVertices.Num(); ++Index)
+		{
+			CurrVertices[Index] = Chaos::FVec3f((*SimulatedVertices)[Index]);
+		}
+
+		const int32 NumVertices = BindingsEval.NumVertices();
+		Positions.SetNum(NumVertices);
+		for (int32 Index = 0; Index < NumVertices; ++Index)
+		{
+			Positions[Index] = BindingsEval.GetEmbeddedPosition(Index, CurrVertices);
+		}
+		
+	}
+}
diff --git a/Engine/Plugins/Experimental/ChaosFlesh/Source/ChaosFleshEngine/Public/ChaosFlesh/FleshCollectionEngineUtility.h b/Engine/Plugins/Experimental/ChaosFlesh/Source/ChaosFleshEngine/Public/ChaosFlesh/FleshCollectionEngineUtility.h
new file mode 100644
--- /dev/null
+++ b/Engine/Plugins/Experimental/ChaosFlesh/Source/ChaosFleshEngine/Public/ChaosFlesh/FleshCollectionEngineUtility.h
@@ -0,0 +1,24 @@
+// Copyright Epic Games, Inc. All Rights Reserved.
+
+#pragma once
+
+#include "Math/MathFwd.h"
+#include "GeometryCollection/ManagedArray.h"
+
+class USkeletalMesh;
+class FFleshCollection;
+class UStaticMesh;
+
+namespace ChaosFlesh 
+{
+	CHAOSFLESHENGINE_API FString GetMeshId(const USkeletalMesh* SkeletalMesh, const bool bUseImportModel);
+
+	CHAOSFLESHENGINE_API FString GetMeshId(const UStaticMesh* StaticMesh);
+
+	CHAOSFLESHENGINE_API void BoundSurfacePositions(
+		const USkeletalMesh* SkeletalMesh,
+		const FFleshCollection* FleshCollection,
+		const TManagedArray<FVector3f>* RestVertices,
+		const TManagedArray<FVector3f>* SimulatedVertices,
+		TArray<FVector3f>& Positions);
+}
diff --git a/Engine/Plugins/Experimental/ChaosFlesh/Source/ChaosFleshNodes/Private/Dataflow/ChaosFleshGenerateSurfaceBindingsNode.cpp b/Engine/Plugins/Experimental/ChaosFlesh/Source/ChaosFleshNodes/Private/Dataflow/ChaosFleshGenerateSurfaceBindingsNode.cpp
--- a/Engine/Plugins/Experimental/ChaosFlesh/Source/ChaosFleshNodes/Private/Dataflow/ChaosFleshGenerateSurfaceBindingsNode.cpp
+++ b/Engine/Plugins/Experimental/ChaosFlesh/Source/ChaosFleshNodes/Private/Dataflow/ChaosFleshGenerateSurfaceBindingsNode.cpp
@@ -8,7 +8,7 @@
 #include "Chaos/BoundingVolumeHierarchy.h"
 #include "Chaos/Tetrahedron.h"
 #include "Chaos/TriangleMesh.h"
-#include "ChaosFlesh/FleshCollectionUtility.h"
+#include "ChaosFlesh/FleshCollectionEngineUtility.h"
 #include "ChaosFlesh/TetrahedralCollection.h"
 #include "Containers/Map.h"
 #include "Engine/StaticMesh.h"
