ChaosFlesh: geometry cache generation takes skeletal mesh as input; GenerateSurfaceBindings takes an optional list of geometry names

diff --git a/Engine/Plugins/Experimental/ChaosFlesh/Source/ChaosFleshEngine/Private/ChaosFlesh/ChaosDeformableTetrahedralComponent.cpp b/Engine/Plugins/Experimental/ChaosFlesh/Source/ChaosFleshEngine/Private/ChaosFlesh/ChaosDeformableTetrahedralComponent.cpp
--- a/Engine/Plugins/Experimental/ChaosFlesh/Source/ChaosFleshEngine/Private/ChaosFlesh/ChaosDeformableTetrahedralComponent.cpp
+++ b/Engine/Plugins/Experimental/ChaosFlesh/Source/ChaosFleshEngine/Private/ChaosFlesh/ChaosDeformableTetrahedralComponent.cpp
@@ -639,13 +639,9 @@ TArray<FVector> UDeformableTetrahedralComponent::GetSkeletalMeshBindingPositions
 	return GetSkeletalMeshBindingPositionsInternal(InSkeletalMesh, nullptr);
 }
 
-TArray<FVector3f> UDeformableTetrahedralComponent::GetGeometryCachePositions(USkeletalMeshComponent* SkeletalComponent) const
+TArray<FVector3f> UDeformableTetrahedralComponent::GetGeometryCachePositions(const USkeletalMesh* SkeletalMesh) const
 {
-	check(SkeletalComponent);
-	SkeletalComponent->RecreateRenderState_Concurrent();
-
 	TArray<FVector3f> Positions;
-	const USkeletalMesh* SkeletalMesh = SkeletalComponent->GetSkeletalMeshAsset();
 	if (RestCollection && DynamicCollection && SkeletalMesh)
 	{
 		const FFleshCollection* FleshCollection = RestCollection->GetCollection();
diff --git a/Engine/Plugins/Experimental/ChaosFlesh/Source/ChaosFleshEngine/Public/ChaosFlesh/ChaosDeformableTetrahedralComponent.h b/Engine/Plugins/Experimental/ChaosFlesh/Source/ChaosFleshEngine/Public/ChaosFlesh/ChaosDeformableTetrahedralComponent.h
--- a/Engine/Plugins/Experimental/ChaosFlesh/Source/ChaosFleshEngine/Public/ChaosFlesh/ChaosDeformableTetrahedralComponent.h
+++ b/Engine/Plugins/Experimental/ChaosFlesh/Source/ChaosFleshEngine/Public/ChaosFlesh/ChaosDeformableTetrahedralComponent.h
@@ -20,7 +20,7 @@ class ADeformableSolverActor;
 class UDeformableSolverComponent;
 class FChaosDeformableTetrahedralSceneProxy;
 class USkinnedAsset;
-class USkeletalMeshComponent;
+class USkeletalMesh;
 
 /**
 *  Options for binding positions query.
@@ -138,7 +138,7 @@ public:
 	TArray<FVector> GetSkeletalMeshBindingPositions(const USkeletalMesh* InSkeletalMesh) const;
 	
 	//~ Begin IDataflowGeometryCachable Interface
-	virtual TArray<FVector3f> GetGeometryCachePositions(USkeletalMeshComponent* SkeletalComponent) const override;
+	virtual TArray<FVector3f> GetGeometryCachePositions(const USkeletalMesh* SkeletalMesh) const override;
 
 	virtual TOptional<TArray<int32>> GetMeshImportVertexMap(const USkinnedAsset& SkinnedMeshAsset) const override;
 	//~ End IDataflowGeometryCachable Interface
diff --git a/Engine/Plugins/Experimental/ChaosFlesh/Source/ChaosFleshNodes/Private/Dataflow/ChaosFleshGenerateSurfaceBindingsNode.cpp b/Engine/Plugins/Experimental/ChaosFlesh/Source/ChaosFleshNodes/Private/Dataflow/ChaosFleshGenerateSurfaceBindingsNode.cpp
--- a/Engine/Plugins/Experimental/ChaosFlesh/Source/ChaosFleshNodes/Private/Dataflow/ChaosFleshGenerateSurfaceBindingsNode.cpp
+++ b/Engine/Plugins/Experimental/ChaosFlesh/Source/ChaosFleshNodes/Private/Dataflow/ChaosFleshGenerateSurfaceBindingsNode.cpp
@@ -21,6 +21,8 @@
 #include "Rendering/SkeletalMeshModel.h"
 #include "Rendering/SkeletalMeshRenderData.h"
 #include "UObject/PrimaryAssetId.h"
+#include "Dataflow/ChaosFleshNodesUtility.h"
+
 DEFINE_LOG_CATEGORY(LogMeshBindings);
 
 void
@@ -128,35 +130,34 @@ BuildVertexToVertexAdjacencyBuffer(
 void
 FGenerateSurfaceBindings::Evaluate(Dataflow::FContext& Context, const FDataflowOutput* Out) const
 {
-	if (Out->IsA<DataType>(&Collection))
+	if (Out->IsA<FManagedArrayCollection>(&Collection))
 	{
-		DataType OutCollection = GetValue<DataType>(Context, &Collection); // Deep copy
-
-		TManagedArray<FIntVector4>* Tetrahedron = 
-			OutCollection.FindAttribute<FIntVector4>(
+		TUniquePtr<FTetrahedralCollection> InCollection(GetValue<FManagedArrayCollection>(Context, &Collection).NewCopy<FTetrahedralCollection>());
+		const TManagedArray<FIntVector4>* Tetrahedron = 
+			InCollection->FindAttribute<FIntVector4>(
 				FTetrahedralCollection::TetrahedronAttribute, FTetrahedralCollection::TetrahedralGroup);
-		TManagedArray<int32>* TetrahedronStart =
-			OutCollection.FindAttribute<int32>(
+		const TManagedArray<int32>* TetrahedronStart =
+			InCollection->FindAttribute<int32>(
 				FTetrahedralCollection::TetrahedronStartAttribute, FGeometryCollection::GeometryGroup);
-		TManagedArray<int32>* TetrahedronCount =
-			OutCollection.FindAttribute<int32>(
+		const TManagedArray<int32>* TetrahedronCount =
+			InCollection->FindAttribute<int32>(
 				FTetrahedralCollection::TetrahedronCountAttribute, FGeometryCollection::GeometryGroup);
-		TManagedArray<TArray<int32>>* IncidentElements =
-			OutCollection.FindAttribute<TArray<int32>>(
+		const TManagedArray<TArray<int32>>* IncidentElements =
+			InCollection->FindAttribute<TArray<int32>>(
 				FTetrahedralCollection::IncidentElementsAttribute, FGeometryCollection::VerticesGroup);
 
-		TManagedArray<FIntVector>* Triangle =
-			OutCollection.FindAttribute<FIntVector>(
+		const TManagedArray<FIntVector>* Triangle =
+			InCollection->FindAttribute<FIntVector>(
 				"Indices", FGeometryCollection::FacesGroup);
-		TManagedArray<int32>* FacesStart =
-			OutCollection.FindAttribute<int32>(
+		const TManagedArray<int32>* FacesStart =
+			InCollection->FindAttribute<int32>(
 				"FaceStart", FGeometryCollection::GeometryGroup);
-		TManagedArray<int32>* FacesCount =
-			OutCollection.FindAttribute<int32>(
+		const TManagedArray<int32>* FacesCount =
+			InCollection->FindAttribute<int32>(
 				"FaceCount", FGeometryCollection::GeometryGroup);
 
-		TManagedArray<FVector3f>* Vertex = 
-			OutCollection.FindAttribute<FVector3f>(
+		const TManagedArray<FVector3f>* Vertex = 
+			InCollection->FindAttribute<FVector3f>(
 				"Vertex", "Vertices");
 
 		TObjectPtr<const USkeletalMesh> SkeletalMesh = GetValue<TObjectPtr<const USkeletalMesh>>(Context, &SkeletalMeshIn);
@@ -271,7 +272,7 @@ FGenerateSurfaceBindings::Evaluate(Dataflow::FContext& Context, const FDataflowO
 				}
 			}
 			TArray<FString> GeometryGroupGuidsLocal = GetValue<TArray<FString>>(Context, &GeometryGroupGuidsIn);
-			TManagedArray<FString>* Guids = OutCollection.FindAttribute<FString>("Guid", FGeometryCollection::GeometryGroup);
+			const TManagedArray<FString>* Guids = InCollection->FindAttribute<FString>("Guid", FGeometryCollection::GeometryGroup);
 
 			// Build Tetrahedra
 			TArray<Chaos::TTetrahedron<Chaos::FReal>> Tets;			// Index 0 == TetMeshStart
@@ -286,7 +287,20 @@ FGenerateSurfaceBindings::Evaluate(Dataflow::FContext& Context, const FDataflowO
 			TConstArrayView<Chaos::TVec3<Chaos::FRealDouble>> VertexDView(VertexD);
 			TArray<Chaos::FVec3> PointNormals;
 
-			for (int32 TetMeshIdx = 0; TetMeshIdx < TetrahedronStart->Num(); TetMeshIdx++)
+			TArray<int32> GeometryGroupSelected;
+			if (IsConnected(&TransformSelection))
+			{
+				FDataflowTransformSelection InTransformSelection = GetValue<FDataflowTransformSelection>(Context, &TransformSelection);
+				GeometryGroupSelected = InCollection->TransformSelectionToGeometryIndices(InTransformSelection.AsArray());
+			}
+			else
+			{
+				for (int32 TetMeshIdx = 0; TetMeshIdx < TetrahedronStart->Num(); ++TetMeshIdx)
+				{
+					GeometryGroupSelected.Add(TetMeshIdx);
+				}
+			}
+			for (const int32& TetMeshIdx: GeometryGroupSelected)
 			{
 				if (GeometryGroupGuidsLocal.Num() && Guids)
 				{
@@ -590,7 +604,7 @@ FGenerateSurfaceBindings::Evaluate(Dataflow::FContext& Context, const FDataflowO
 			} // end for all LOD
 
 			// Stash bindings in the geometry collection
-			GeometryCollection::Facades::FTetrahedralBindings TetBindings(OutCollection);
+			GeometryCollection::Facades::FTetrahedralBindings TetBindings(*InCollection);
 			TetBindings.DefineSchema();
 			FName MeshName(*MeshId, MeshId.Len());
 			for (int32 LOD = 0; LOD < MeshVertices.Num(); LOD++)
@@ -599,6 +613,6 @@ FGenerateSurfaceBindings::Evaluate(Dataflow::FContext& Context, const FDataflowO
 				TetBindings.SetBindingsData(Parents[LOD], Weights[LOD], Offsets[LOD], Masks[LOD]);
 			}
 		}
-		SetValue(Context, MoveTemp(OutCollection), &Collection);
+		SetValue<const FManagedArrayCollection&>(Context, *InCollection, &Collection);
 	}
 }
diff --git a/Engine/Plugins/Experimental/ChaosFlesh/Source/ChaosFleshNodes/Private/Dataflow/ChaosFleshSetFleshDefaultPropertiesNode.cpp b/Engine/Plugins/Experimental/ChaosFlesh/Source/ChaosFleshNodes/Private/Dataflow/ChaosFleshSetFleshDefaultPropertiesNode.cpp
--- a/Engine/Plugins/Experimental/ChaosFlesh/Source/ChaosFleshNodes/Private/Dataflow/ChaosFleshSetFleshDefaultPropertiesNode.cpp
+++ b/Engine/Plugins/Experimental/ChaosFlesh/Source/ChaosFleshNodes/Private/Dataflow/ChaosFleshSetFleshDefaultPropertiesNode.cpp
@@ -62,6 +62,7 @@ void FSetFleshDefaultPropertiesNode::Evaluate(Dataflow::FContext& Context, const
 				const TManagedArray<TArray<int32>>& IncidentElements = InCollection.GetAttribute<TArray<int32>>(FTetrahedralCollection::IncidentElementsAttribute, FGeometryCollection::VerticesGroup);
 				const TManagedArray<TArray<int32>>& IncidentElementsLocalIndex = InCollection.GetAttribute<TArray<int32>>(FTetrahedralCollection::IncidentElementsLocalIndexAttribute, FGeometryCollection::VerticesGroup);
 				
+				Mass.Fill(0.f);
 				float MinV = TNumericLimits<float>::Max();
 				float MaxV = -TNumericLimits<float>::Max();
 				double AvgV = 0.0;
diff --git a/Engine/Plugins/Experimental/ChaosFlesh/Source/ChaosFleshNodes/Public/Dataflow/ChaosFleshGenerateSurfaceBindingsNode.h b/Engine/Plugins/Experimental/ChaosFlesh/Source/ChaosFleshNodes/Public/Dataflow/ChaosFleshGenerateSurfaceBindingsNode.h
--- a/Engine/Plugins/Experimental/ChaosFlesh/Source/ChaosFleshNodes/Public/Dataflow/ChaosFleshGenerateSurfaceBindingsNode.h
+++ b/Engine/Plugins/Experimental/ChaosFlesh/Source/ChaosFleshNodes/Public/Dataflow/ChaosFleshGenerateSurfaceBindingsNode.h
@@ -5,10 +5,9 @@
 #include "CoreMinimal.h"
 #include "Dataflow/DataflowEngine.h"
 #include "GeometryCollection/ManagedArrayCollection.h"
-
 #include "Engine/StaticMesh.h"
 #include "Engine/SkeletalMesh.h"
-
+#include "Dataflow/DataflowSelection.h"
 #include "ChaosFleshGenerateSurfaceBindingsNode.generated.h"
 
 class UStaticMesh;
@@ -25,8 +24,6 @@ struct FGenerateSurfaceBindings : public FDataflowNode
 	DATAFLOW_NODE_RENDER_TYPE("SurfaceRender", FGeometryCollection::StaticType(), "Collection")
 
 public:
-	typedef FManagedArrayCollection DataType;
-
 	/** Passthrough geometry collection. Bindings are stored as standalone groups in the \p Collection, keyed by the name of the input render mesh and all available LOD's. */
 	UPROPERTY(meta = (DataflowInput, DataflowOutput, DataflowPassthrough = "Collection", DisplayName = "Collection"))
 	FManagedArrayCollection Collection;
@@ -39,13 +36,18 @@ public:
 	UPROPERTY(EditAnywhere, Category = "Dataflow", meta = (DataflowInput, DisplayName = "SkeletalMesh"))
 	TObjectPtr<const USkeletalMesh> SkeletalMeshIn = nullptr;
 
+	/** Only meshes with transforms in TransformSelection will be bound. */
+	UPROPERTY(meta = (DataflowInput, DisplayName = "(Optional) TransformSelection"))
+	FDataflowTransformSelection TransformSelection;
+
+	/** Only meshes with transforms in TransformSelection will be bound. */
+	UPROPERTY(meta = (DataflowInput, DisplayName = "(Optional) GeometryGroupGuids"))
+	TArray<FString> GeometryGroupGuidsIn;
+
 	/** Use the import geometry of the skeletal mesh. */
 	UPROPERTY(EditAnywhere, Category = "Dataflow", meta = (DisplayName = "UseSkeletalMeshImportModel"))
 	bool bUseSkeletalMeshImportModel = false;
 
-	UPROPERTY(meta = (DataflowInput, DisplayName = "(Optional)GeometryGroupGuidsIn"))
-	TArray<FString> GeometryGroupGuidsIn;
-
 	/** Enable binding to the exterior hull of the tetrahedron mesh. */
 	UPROPERTY(EditAnywhere, Category = "Dataflow", meta = (DisplayName = "DoSurfaceProjection"))
 	bool bDoSurfaceProjection = true;
@@ -65,6 +67,7 @@ public:
 		RegisterOutputConnection(&Collection, &Collection);
 		RegisterInputConnection(&StaticMeshIn);
 		RegisterInputConnection(&SkeletalMeshIn);
+		RegisterInputConnection(&TransformSelection);
 		RegisterInputConnection(&GeometryGroupGuidsIn);
 	}
 
diff --git a/Engine/Plugins/Experimental/Dataflow/Source/DataflowEditor/Private/Dataflow/DataflowSimulationScene.cpp b/Engine/Plugins/Experimental/Dataflow/Source/DataflowEditor/Private/Dataflow/DataflowSimulationScene.cpp
--- a/Engine/Plugins/Experimental/Dataflow/Source/DataflowEditor/Private/Dataflow/DataflowSimulationScene.cpp
+++ b/Engine/Plugins/Experimental/Dataflow/Source/DataflowEditor/Private/Dataflow/DataflowSimulationScene.cpp
@@ -314,25 +314,17 @@ void UDataflowSimulationSceneDescription::GenerateGeometryCache()
 	float DeltaTime = (TimeRange[1] - TimeRange[0]) / NumFrames;
 	TObjectPtr<AActor> GetRootActor = SimulationScene->GetRootActor();
 	TObjectPtr<AActor> PreviewActor = SimulationScene->GetPreviewActor();
-	if (CacheAsset && GeometryCacheAsset && GetRootActor)
+	if (CacheAsset && GeometryCacheAsset && GetRootActor && EmbeddedSkeletalMesh)
 	{
 		IDataflowGeometryCachable* GeometryCachable = nullptr; //interface for ChaosDeformableTetrahedralComponent
 
 		RenderPositions.SetNum(NumFrames);
 		TInlineComponentArray<UPrimitiveComponent*> PrimComponents;
 		PreviewActor->GetComponents(PrimComponents);
-		USkeletalMeshComponent* SkeletalComponent = nullptr;
 		for (UPrimitiveComponent* PrimComponent : PrimComponents)
 		{
-			if (!GeometryCachable)
-			{
-				GeometryCachable = Cast<IDataflowGeometryCachable>(PrimComponent);
-			}
-			if (!SkeletalComponent)
-			{
-				SkeletalComponent = Cast<USkeletalMeshComponent>(PrimComponent);
-			}
-			if (GeometryCachable && SkeletalComponent)
+			GeometryCachable = Cast<IDataflowGeometryCachable>(PrimComponent);
+			if (GeometryCachable)
 			{
 				break;
 			}
@@ -342,12 +334,7 @@ void UDataflowSimulationSceneDescription::GenerateGeometryCache()
 			UE_LOG(LogChaosDataflow, Error, TEXT("No Flesh Component in the Preview Actor"));
 			return;
 		}
-		else if (!SkeletalComponent)
-		{
-			UE_LOG(LogChaosDataflow, Error, TEXT("No Skeletal Mesh Component in the Preview Actor"));
-			return;
-		}
-		TOptional<TArray<int32>> OptionalMap = GeometryCachable->GetMeshImportVertexMap(*SkeletalComponent->GetSkeletalMeshAsset());
+		TOptional<TArray<int32>> OptionalMap = GeometryCachable->GetMeshImportVertexMap(*EmbeddedSkeletalMesh);
 		if (!OptionalMap)
 		{
 			return;
@@ -358,9 +345,9 @@ void UDataflowSimulationSceneDescription::GenerateGeometryCache()
 		{
 			Time += DeltaTime;
 			Cast<AChaosCacheManager>(GetRootActor)->SetStartTime(Time);
-			RenderPositions[Frame] = GeometryCachable->GetGeometryCachePositions(SkeletalComponent);
+			RenderPositions[Frame] = GeometryCachable->GetGeometryCachePositions(EmbeddedSkeletalMesh);
 		}
-		DataflowSimulationGeometryCache::SaveGeometryCache(*GeometryCacheAsset, *SkeletalComponent->GetSkeletalMeshAsset(), ImportedVertexNumbers, RenderPositions);
+		DataflowSimulationGeometryCache::SaveGeometryCache(*GeometryCacheAsset, *EmbeddedSkeletalMesh, ImportedVertexNumbers, RenderPositions);
 		DataflowSimulationGeometryCache::SavePackage(*GeometryCacheAsset);
 	}
 }
diff --git a/Engine/Plugins/Experimental/Dataflow/Source/DataflowEditor/Public/Dataflow/DataflowSimulationScene.h b/Engine/Plugins/Experimental/Dataflow/Source/DataflowEditor/Public/Dataflow/DataflowSimulationScene.h
--- a/Engine/Plugins/Experimental/Dataflow/Source/DataflowEditor/Public/Dataflow/DataflowSimulationScene.h
+++ b/Engine/Plugins/Experimental/Dataflow/Source/DataflowEditor/Public/Dataflow/DataflowSimulationScene.h
@@ -53,12 +53,16 @@ public:
 	UPROPERTY(EditAnywhere, Category = "Geometry", DisplayName="GeometryCache", meta=(EditCondition = "CacheAsset != nullptr"))
 	TObjectPtr<UGeometryCache> GeometryCacheAsset = nullptr;
 
+	/** SkeletalMesh interpolated from simulation. This should match the SkeletalMesh used in GenerateSurfaceBindings node */
+	UPROPERTY(EditAnywhere, Category = "Geometry", DisplayName = "EmbeddedSkeletalMesh", meta = (EditCondition = "CacheAsset != nullptr"))
+	TObjectPtr<USkeletalMesh> EmbeddedSkeletalMesh = nullptr;
+
 	/** Interpolates and saves geometry cache from Chaos cache */
-	UFUNCTION(CallInEditor, Category = "Geometry")
+	UFUNCTION(CallInEditor, Category = "Geometry", meta = (EditCondition = "CacheAsset != nullptr && GeometryCacheAsset != nullptr && EmbeddedSkeletalMesh != nullptr"))
 	void GenerateGeometryCache();
 
 	/** Creates a new geometry cache file */
-	UFUNCTION(CallInEditor, Category = "Geometry")
+	UFUNCTION(CallInEditor, Category = "Geometry", meta = (EditCondition = "CacheAsset != nullptr && GeometryCacheAsset != nullptr && EmbeddedSkeletalMesh != nullptr"))
 	void NewGeometryCache();
 
 private:
diff --git a/Engine/Source/Runtime/Experimental/Chaos/Private/GeometryCollection/GeometryCollection.cpp b/Engine/Source/Runtime/Experimental/Chaos/Private/GeometryCollection/GeometryCollection.cpp
--- a/Engine/Source/Runtime/Experimental/Chaos/Private/GeometryCollection/GeometryCollection.cpp
+++ b/Engine/Source/Runtime/Experimental/Chaos/Private/GeometryCollection/GeometryCollection.cpp
@@ -2166,5 +2166,16 @@ void FGeometryCollection::UpdateOldAttributeNames()
 	this->RemoveGroup("Structure");
 }
 
-
+TArray<int32> FGeometryCollection::TransformSelectionToGeometryIndices(const TArray<int32>& Transforms)
+{
+	TArray<int32> Geometries;
+	for (const int32& TransformIdx : Transforms)
+	{
+		if (TransformToGeometryIndex.IsValidIndex(TransformIdx)	&& IsGeometry(TransformIdx))
+		{
+			Geometries.Add(TransformToGeometryIndex[TransformIdx]);
+		}
+	}
+	return Geometries;
+}
 
diff --git a/Engine/Source/Runtime/Experimental/Chaos/Public/GeometryCollection/GeometryCollection.h b/Engine/Source/Runtime/Experimental/Chaos/Public/GeometryCollection/GeometryCollection.h
--- a/Engine/Source/Runtime/Experimental/Chaos/Public/GeometryCollection/GeometryCollection.h
+++ b/Engine/Source/Runtime/Experimental/Chaos/Public/GeometryCollection/GeometryCollection.h
@@ -394,7 +394,7 @@ public:
 	/* Backwards compatibility */
 	CHAOS_API void UpdateOldAttributeNames();
 
-
+	CHAOS_API TArray<int32> TransformSelectionToGeometryIndices(const TArray<int32>& Transforms);
 };
 
 FORCEINLINE Chaos::FChaosArchive& operator<<(Chaos::FChaosArchive& Ar, FGeometryCollection& Value)
diff --git a/Engine/Source/Runtime/Experimental/Dataflow/Simulation/Public/Dataflow/Interfaces/DataflowInterfaceGeometryCachable.h b/Engine/Source/Runtime/Experimental/Dataflow/Simulation/Public/Dataflow/Interfaces/DataflowInterfaceGeometryCachable.h
--- a/Engine/Source/Runtime/Experimental/Dataflow/Simulation/Public/Dataflow/Interfaces/DataflowInterfaceGeometryCachable.h
+++ b/Engine/Source/Runtime/Experimental/Dataflow/Simulation/Public/Dataflow/Interfaces/DataflowInterfaceGeometryCachable.h
@@ -9,7 +9,7 @@
 #include "Math/MathFwd.h"
 #include "DataflowInterfaceGeometryCachable.generated.h"
 
-class USkeletalMeshComponent;
+class USkeletalMesh;
 class USkinnedAsset;
 
 UINTERFACE(meta = (CannotImplementInterfaceInBlueprint), MinimalAPI)
@@ -25,7 +25,7 @@ class IDataflowGeometryCachable
 
 public:
 	/** Computes skeletal mesh positions for geometry cache */
-	virtual TArray<FVector3f> GetGeometryCachePositions(USkeletalMeshComponent* SkeletalComponent) const = 0;
+	virtual TArray<FVector3f> GetGeometryCachePositions(const USkeletalMesh* SkeletalMesh) const = 0;
 
 	/** Get mapping from final mesh vertex index to raw import vertex index */
 	virtual TOptional<TArray<int32>> GetMeshImportVertexMap(const USkinnedAsset& SkinnedMeshAsset) const = 0;
