ChaosFlesh: InvertCurrent in painting tool; use VertexSelection instead of TArray; add AttributeKey input

diff --git a/Engine/Plugins/Experimental/ChaosFlesh/Source/ChaosFleshNodes/Private/Dataflow/ChaosFleshGetSurfaceIndicesNode.cpp b/Engine/Plugins/Experimental/ChaosFlesh/Source/ChaosFleshNodes/Private/Dataflow/ChaosFleshGetSurfaceIndicesNode.cpp
--- a/Engine/Plugins/Experimental/ChaosFlesh/Source/ChaosFleshNodes/Private/Dataflow/ChaosFleshGetSurfaceIndicesNode.cpp
+++ b/Engine/Plugins/Experimental/ChaosFlesh/Source/ChaosFleshNodes/Private/Dataflow/ChaosFleshGetSurfaceIndicesNode.cpp
@@ -6,45 +6,55 @@
 
 void FGetSurfaceIndicesNode::Evaluate(Dataflow::FContext& Context, const FDataflowOutput* Out) const
 {
-	FManagedArrayCollection InCollection = GetValue<FManagedArrayCollection>(Context, &Collection);
-	TArray<int32> SurfaceIndicesLocal;
-	if (TManagedArray<FIntVector>* Indices = InCollection.FindAttribute<FIntVector>("Indices", FGeometryCollection::FacesGroup))
+	const FManagedArrayCollection InCollection = GetValue<FManagedArrayCollection>(Context, &Collection);
+	if (const TManagedArray<FVector3f>* Vertex = InCollection.FindAttribute<FVector3f>("Vertex", FGeometryCollection::VerticesGroup))
 	{
-		if (FindInput(&GeometryGroupGuidsIn) && FindInput(&GeometryGroupGuidsIn)->GetConnection())
+		FDataflowVertexSelection SurfaceVertexSelectionOut;
+		SurfaceVertexSelectionOut.Initialize(Vertex->Num(), false);
+		TArray<int32> SurfaceIndicesLocal;
+		if (const TManagedArray<FIntVector>* Indices = InCollection.FindAttribute<FIntVector>("Indices", FGeometryCollection::FacesGroup))
 		{
-			TArray<FString> GeometryGroupGuidsLocal = GetValue<TArray<FString>>(Context, &GeometryGroupGuidsIn);
-			TManagedArray<int32>* IndicesStart = InCollection.FindAttribute<int32>("FaceStart", FGeometryCollection::GeometryGroup);
-			TManagedArray<int32>* IndicesCount = InCollection.FindAttribute<int32>("FaceCount", FGeometryCollection::GeometryGroup);
-			if (TManagedArray<FString>* Guids = InCollection.FindAttribute<FString>("Guid", FGeometryCollection::GeometryGroup))
+			if (FindInput(&GeometryGroupGuidsIn) && FindInput(&GeometryGroupGuidsIn)->GetConnection())
 			{
-				for (int32 Idx = 0; Idx < IndicesStart->Num(); Idx++)
+				const TArray<FString> GeometryGroupGuidsLocal = GetValue<TArray<FString>>(Context, &GeometryGroupGuidsIn);
+				const TManagedArray<int32>* IndicesStart = InCollection.FindAttribute<int32>("FaceStart", FGeometryCollection::GeometryGroup);
+				const TManagedArray<int32>* IndicesCount = InCollection.FindAttribute<int32>("FaceCount", FGeometryCollection::GeometryGroup);
+				if (const TManagedArray<FString>* Guids = InCollection.FindAttribute<FString>("Guid", FGeometryCollection::GeometryGroup))
 				{
-					if (GeometryGroupGuidsLocal.Num() && Guids)
+					for (int32 Idx = 0; Idx < IndicesStart->Num(); Idx++)
 					{
-						if (GeometryGroupGuidsLocal.Contains((*Guids)[Idx]))
+						if (GeometryGroupGuidsLocal.Num() && Guids)
 						{
-							for (int32 i = (*IndicesStart)[Idx]; i < (*IndicesStart)[Idx] + (*IndicesCount)[Idx]; i++)
+							if (GeometryGroupGuidsLocal.Contains((*Guids)[Idx]))
 							{
-								for (int32 j = 0; j < 3; j++)
+								for (int32 i = (*IndicesStart)[Idx]; i < (*IndicesStart)[Idx] + (*IndicesCount)[Idx]; i++)
 								{
-									SurfaceIndicesLocal.AddUnique((*Indices)[i][j]);
+									for (int32 j = 0; j < 3; j++)
+									{
+										SurfaceIndicesLocal.AddUnique((*Indices)[i][j]);
+									}
 								}
 							}
 						}
 					}
 				}
 			}
-		}
-		else
-		{
-			for (int32 i = 0; i < Indices->Num(); i++)
+			else
 			{
-				for (int32 j = 0; j < 3; j++)
+				for (int32 i = 0; i < Indices->Num(); i++)
 				{
-					SurfaceIndicesLocal.AddUnique((*Indices)[i][j]);
+					for (int32 j = 0; j < 3; j++)
+					{
+						SurfaceIndicesLocal.AddUnique((*Indices)[i][j]);
+					}
 				}
 			}
+			SurfaceVertexSelectionOut.SetFromArray(SurfaceIndicesLocal);
 		}
+		SetValue(Context, MoveTemp(SurfaceVertexSelectionOut), &SurfaceVertexSelection);
+	}
+	else
+	{
+		SetValue(Context, FDataflowVertexSelection(), &SurfaceVertexSelection);
 	}
-	SetValue(Context, MoveTemp(SurfaceIndicesLocal), &SurfaceIndicesOut);
 }
diff --git a/Engine/Plugins/Experimental/ChaosFlesh/Source/ChaosFleshNodes/Private/Dataflow/ChaosFleshSetFleshBonePositionTargetBindingNode.cpp b/Engine/Plugins/Experimental/ChaosFlesh/Source/ChaosFleshNodes/Private/Dataflow/ChaosFleshSetFleshBonePositionTargetBindingNode.cpp
--- a/Engine/Plugins/Experimental/ChaosFlesh/Source/ChaosFleshNodes/Private/Dataflow/ChaosFleshSetFleshBonePositionTargetBindingNode.cpp
+++ b/Engine/Plugins/Experimental/ChaosFlesh/Source/ChaosFleshNodes/Private/Dataflow/ChaosFleshSetFleshBonePositionTargetBindingNode.cpp
@@ -24,304 +24,288 @@ void FSetFleshBonePositionTargetBindingDataflowNode::Evaluate(Dataflow::FContext
 
 		if (TObjectPtr<USkeletalMesh> BoneSkeletalMesh = GetValue<TObjectPtr<USkeletalMesh>>(Context, &SkeletalMeshIn))
 		{
-			if (TManagedArray<FIntVector>* Indices = InCollection.FindAttribute<FIntVector>("Indices", FGeometryCollection::FacesGroup))
+			const TManagedArray<FIntVector>* Indices = InCollection.FindAttribute<FIntVector>("Indices", FGeometryCollection::FacesGroup);
+			const TManagedArray<FVector3f>* Vertices = InCollection.FindAttribute<FVector3f>("Vertex", FGeometryCollection::VerticesGroup);
+			const TManagedArray<FTransform3f>* Transform = InCollection.FindAttribute<FTransform3f>("Transform", FTransformCollection::TransformGroup);
+			const TManagedArray<FString>* TransformBoneName = InCollection.FindAttribute<FString>("BoneName", FTransformCollection::TransformGroup);
+			if (Indices && Vertices && Transform && TransformBoneName)
 			{
-				if (TManagedArray<FVector3f>* Vertices = InCollection.FindAttribute<FVector3f>("Vertex", FGeometryCollection::VerticesGroup))
+				TMap<FString, int32> BoneNameIndexMap;
+				for (int32 BoneIndex = 0; BoneIndex < TransformBoneName->Num(); ++BoneIndex)
 				{
-
-					FSkeletalMeshRenderData* RenderData = BoneSkeletalMesh->GetResourceForRendering();
-					if (RenderData->LODRenderData.Num())
+					BoneNameIndexMap.Add((*TransformBoneName)[BoneIndex], BoneIndex);
+				}
+				FSkeletalMeshRenderData* RenderData = BoneSkeletalMesh->GetResourceForRendering();
+				if (RenderData->LODRenderData.Num())
+				{
+					FSkeletalMeshLODRenderData* LODRenderData = &RenderData->LODRenderData[0];
+					const FPositionVertexBuffer& PositionVertexBuffer =
+						LODRenderData->StaticVertexBuffers.PositionVertexBuffer;
+
+					const FSkinWeightVertexBuffer* SkinWeightVertexBuffer = LODRenderData->GetSkinWeightVertexBuffer();
+					const int32 MaxBoneInfluences = SkinWeightVertexBuffer->GetMaxBoneInfluences();
+					TArray<FTransform> ComponentPose;
+					Dataflow::Animation::GlobalTransforms(BoneSkeletalMesh->GetRefSkeleton(), ComponentPose);
+					TArray<TArray<int32>> BoneBoundVerts;
+					TArray<TArray<float>> BoneBoundWeights;
+					BoneBoundVerts.SetNum(ComponentPose.Num());
+					BoneBoundWeights.SetNum(ComponentPose.Num());
+
+					int32 NumSkeletonVertices = int32(PositionVertexBuffer.GetNumVertices());
+
+					Chaos::FReal SphereRadius = (Chaos::FReal)0.;
+
+					Chaos::TVec3<float> CoordMaxs(-FLT_MAX);
+					Chaos::TVec3<float> CoordMins(FLT_MAX);
+					for (int32 i = 0; i < NumSkeletonVertices; i++)
 					{
-						FSkeletalMeshLODRenderData* LODRenderData = &RenderData->LODRenderData[0];
-						const FPositionVertexBuffer& PositionVertexBuffer =
-							LODRenderData->StaticVertexBuffers.PositionVertexBuffer;
-
-						const FSkinWeightVertexBuffer* SkinWeightVertexBuffer = LODRenderData->GetSkinWeightVertexBuffer();
-						const int32 MaxBoneInfluences = SkinWeightVertexBuffer->GetMaxBoneInfluences();
-						TArray<FTransform> ComponentPose;
-						Dataflow::Animation::GlobalTransforms(BoneSkeletalMesh->GetRefSkeleton(), ComponentPose);
-						TArray<TArray<int32>> BoneBoundVerts;
-						TArray<TArray<float>> BoneBoundWeights;
-						BoneBoundVerts.SetNum(ComponentPose.Num());
-						BoneBoundWeights.SetNum(ComponentPose.Num());
-
-						uint32 NumSkeletonVertices = PositionVertexBuffer.GetNumVertices();
-
-						Chaos::FReal SphereRadius = (Chaos::FReal)0.;
-
-						Chaos::TVec3<float> CoordMaxs(-FLT_MAX);
-						Chaos::TVec3<float> CoordMins(FLT_MAX);
-						for (int32 i = 0; i < int32(NumSkeletonVertices); i++)
+						for (int32 j = 0; j < 3; j++)
 						{
-							for (int32 j = 0; j < 3; j++)
+							if (PositionVertexBuffer.VertexPosition(i)[j] > CoordMaxs[j])
 							{
-								if (PositionVertexBuffer.VertexPosition(i)[j] > CoordMaxs[j])
-								{
-									CoordMaxs[j] = PositionVertexBuffer.VertexPosition(i)[j];
-								}
-								if (PositionVertexBuffer.VertexPosition(i)[j] < CoordMins[j])
-								{
-									CoordMins[j] = PositionVertexBuffer.VertexPosition(i)[j];
-								}
+								CoordMaxs[j] = PositionVertexBuffer.VertexPosition(i)[j];
+							}
+							if (PositionVertexBuffer.VertexPosition(i)[j] < CoordMins[j])
+							{
+								CoordMins[j] = PositionVertexBuffer.VertexPosition(i)[j];
 							}
 						}
-						Chaos::TVec3<float> CoordDiff = (CoordMaxs - CoordMins) * VertexRadiusRatio;
-						SphereRadius = Chaos::FReal(FGenericPlatformMath::Max(CoordDiff[0], FGenericPlatformMath::Max(CoordDiff[1], CoordDiff[2])));
+					}
+					Chaos::TVec3<float> CoordDiff = (CoordMaxs - CoordMins) * VertexRadiusRatio;
+					SphereRadius = Chaos::FReal(FGenericPlatformMath::Max(CoordDiff[0], FGenericPlatformMath::Max(CoordDiff[1], CoordDiff[2])));
 
-						TArray<Chaos::TSphere<Chaos::FReal, 3>*> VertexSpherePtrs;
-						TArray<Chaos::TSphere<Chaos::FReal, 3>> VertexSpheres;
+					TArray<Chaos::TSphere<Chaos::FReal, 3>*> VertexSpherePtrs;
+					TArray<Chaos::TSphere<Chaos::FReal, 3>> VertexSpheres;
 
-						VertexSpheres.Init(Chaos::TSphere<Chaos::FReal, 3>(Chaos::TVec3<Chaos::FReal>(0), SphereRadius), NumSkeletonVertices);
-						VertexSpherePtrs.SetNum(NumSkeletonVertices);
+					VertexSpheres.Init(Chaos::TSphere<Chaos::FReal, 3>(Chaos::TVec3<Chaos::FReal>(0), SphereRadius), NumSkeletonVertices);
+					VertexSpherePtrs.SetNum(NumSkeletonVertices);
 
-						for (int32 i = 0; i < int32(NumSkeletonVertices); i++)
-						{
-							Chaos::TVec3<Chaos::FReal> SphereCenter(PositionVertexBuffer.VertexPosition(i));
-							Chaos::TSphere<Chaos::FReal, 3> VertexSphere(SphereCenter, SphereRadius);
-							VertexSpheres[i] = Chaos::TSphere<Chaos::FReal, 3>(SphereCenter, SphereRadius);
-							VertexSpherePtrs[i] = &VertexSpheres[i];
-						}
-						Chaos::TBoundingVolumeHierarchy<
-							TArray<Chaos::TSphere<Chaos::FReal, 3>*>,
-							TArray<int32>,
-							Chaos::FReal,
-							3> VertexBVH(VertexSpherePtrs);
+					for (int32 i = 0; i < int32(NumSkeletonVertices); i++)
+					{
+						Chaos::TVec3<Chaos::FReal> SphereCenter(PositionVertexBuffer.VertexPosition(i));
+						Chaos::TSphere<Chaos::FReal, 3> VertexSphere(SphereCenter, SphereRadius);
+						VertexSpheres[i] = Chaos::TSphere<Chaos::FReal, 3>(SphereCenter, SphereRadius);
+						VertexSpherePtrs[i] = &VertexSpheres[i];
+					}
+					Chaos::TBoundingVolumeHierarchy<
+						TArray<Chaos::TSphere<Chaos::FReal, 3>*>,
+						TArray<int32>,
+						Chaos::FReal,
+						3> VertexBVH(VertexSpherePtrs);
 
-						if (SkeletalBindingMode == ESkeletalBindingMode::Dataflow_SkeletalBinding_Kinematic)
-						{
+					if (SkeletalBindingMode == ESkeletalBindingMode::Dataflow_SkeletalBinding_Kinematic)
+					{
 
-							//avoid particles that are already kinematic:
-							TArray<bool> ParticleIsKinematic;
-							ParticleIsKinematic.Init(false, Vertices->Num());
-							typedef GeometryCollection::Facades::FKinematicBindingFacade FKinematics;
-							FKinematics Kinematics(InCollection);
+						GeometryCollection::Facades::FVertexBoneWeightsFacade VertexBoneWeightsFacade(InCollection);
 
-							// Add Kinematics Node
-							for (int i = Kinematics.NumKinematicBindings() - 1; i >= 0; i--)
+						TArray<Chaos::TVector<int32, 3>> IndicesArray;
+						for (int32 i = 0; i < Indices->Num(); i++)
+						{
+							Chaos::TVector<int32, 3> CurrentIndices(0);
+							for (int32 j = 0; j < 3; j++)
 							{
-								FKinematics::FBindingKey Key = Kinematics.GetKinematicBindingKey(i);
-
-								int32 BoneIndex = INDEX_NONE;
-								TArray<int32> BoundVerts;
-								TArray<float> BoundWeights;
-								Kinematics.GetBoneBindings(Key, BoneIndex, BoundVerts, BoundWeights);
-
-								for (int32 vdx : BoundVerts)
-								{
-									ParticleIsKinematic[vdx] = true;
-								}
+								CurrentIndices[j] = (*Indices)[i][j];
 							}
-
-
-							TArray<Chaos::TVector<int32, 3>> IndicesArray;
-							for (int32 i = 0; i < Indices->Num(); i++)
+							if (CurrentIndices[0] != -1
+								&& CurrentIndices[1] != -1
+								&& CurrentIndices[2] != -1)
 							{
-								Chaos::TVector<int32, 3> CurrentIndices(0);
-								for (int32 j = 0; j < 3; j++)
-								{
-									CurrentIndices[j] = (*Indices)[i][j];
-								}
-								if (CurrentIndices[0] != -1
-									&& CurrentIndices[1] != -1
-									&& CurrentIndices[2] != -1)
-								{
-									IndicesArray.Emplace(CurrentIndices);
-								}
+								IndicesArray.Emplace(CurrentIndices);
 							}
-							TArray<TArray<int32>> LocalIndex;
-							TArray<TArray<int32>>* LocalIndexPtr = &LocalIndex;
-							TArray<TArray<int>> GlobalIndex = Chaos::Utilities::ComputeIncidentElements(IndicesArray, LocalIndexPtr);
-							int32 ActualParticleCount = 0;
-							for (int32 l = 0; l < GlobalIndex.Num(); l++)
+						}
+						TArray<TArray<int32>> LocalIndex;
+						TArray<TArray<int32>>* LocalIndexPtr = &LocalIndex;
+						TArray<TArray<int>> GlobalIndex = Chaos::Utilities::ComputeIncidentElements(IndicesArray, LocalIndexPtr);
+						int32 ActualParticleCount = 0;
+						for (int32 l = 0; l < GlobalIndex.Num(); l++)
+						{
+							if (GlobalIndex[l].Num() > 0)
 							{
-								if (GlobalIndex[l].Num() > 0)
-								{
-									ActualParticleCount += 1;
-								}
+								ActualParticleCount += 1;
 							}
-							TArray<Chaos::TVector<float, 3>> IndicesPositions;
-							IndicesPositions.SetNum(ActualParticleCount);
-							TArray<int32> IndicesMap;
-							IndicesMap.SetNum(ActualParticleCount);
-							int32 CurrentParticleIndex = 0;
-							for (int32 i = 0; i < GlobalIndex.Num(); i++)
+						}
+						TArray<Chaos::TVector<float, 3>> IndicesPositions;
+						IndicesPositions.SetNum(ActualParticleCount);
+						TArray<int32> IndicesMap;
+						IndicesMap.SetNum(ActualParticleCount);
+						int32 CurrentParticleIndex = 0;
+						for (int32 i = 0; i < GlobalIndex.Num(); i++)
+						{
+							if (GlobalIndex[i].Num() > 0)
 							{
-								if (GlobalIndex[i].Num() > 0)
-								{
-									IndicesPositions[CurrentParticleIndex] = (*Vertices)[(*Indices)[GlobalIndex[i][0]][LocalIndex[i][0]]];
-									IndicesMap[CurrentParticleIndex] = (*Indices)[GlobalIndex[i][0]][LocalIndex[i][0]];
-									CurrentParticleIndex += 1;
-								}
+								IndicesPositions[CurrentParticleIndex] = (*Vertices)[(*Indices)[GlobalIndex[i][0]][LocalIndex[i][0]]];
+								IndicesMap[CurrentParticleIndex] = (*Indices)[GlobalIndex[i][0]][LocalIndex[i][0]];
+								CurrentParticleIndex += 1;
 							}
+						}
 
-
-							//for (int32 i = 0; i < Vertices->Num(); i++)
-							for (int32 i = 0; i < IndicesMap.Num(); i++)
+						for (int32 i = 0; i < IndicesMap.Num(); i++)
+						{
+							//only work on particles that are not kinematic:
+							if (!VertexBoneWeightsFacade.IsKinematicVertex(IndicesMap[i]))
 							{
-								//only work on particles that are not kinematic:
-								if (!ParticleIsKinematic[IndicesMap[i]])
+								TArray<int32> ParticleIntersection = VertexBVH.FindAllIntersections((*Vertices)[IndicesMap[i]]);
+								int32 MinIndex = -1;
+								float MinDis = SphereRadius;
+								for (int32 j = 0; j < ParticleIntersection.Num(); j++)
 								{
-									TArray<int32> ParticleIntersection = VertexBVH.FindAllIntersections((*Vertices)[IndicesMap[i]]);
-									int32 MinIndex = -1;
-									float MinDis = SphereRadius;
-									for (int32 j = 0; j < ParticleIntersection.Num(); j++)
+									Chaos::FRealSingle CurrentDistance = ((*Vertices)[IndicesMap[i]] - PositionVertexBuffer.VertexPosition(ParticleIntersection[j])).Size();
+									if (CurrentDistance < MinDis)
 									{
-										Chaos::FRealSingle CurrentDistance = ((*Vertices)[IndicesMap[i]] - PositionVertexBuffer.VertexPosition(ParticleIntersection[j])).Size();
-										if (CurrentDistance < MinDis)
-										{
-											MinDis = CurrentDistance;
-											MinIndex = ParticleIntersection[j];
-										}
+										MinDis = CurrentDistance;
+										MinIndex = ParticleIntersection[j];
 									}
+								}
 
-									if (MinIndex != -1)
-									{
-										int32 BoneParticleIndex = -1;
+								if (MinIndex != -1)
+								{
+									int32 BoneParticleIndex = -1;
 
-										int32 SectionIndex;
-										int32 VertIndex;
-										LODRenderData->GetSectionFromVertexIndex(MinIndex, SectionIndex, VertIndex);
+									int32 SectionIndex;
+									int32 VertIndex;
+									LODRenderData->GetSectionFromVertexIndex(MinIndex, SectionIndex, VertIndex);
 
-										check(SectionIndex < LODRenderData->RenderSections.Num());
-										const FSkelMeshRenderSection& Section = LODRenderData->RenderSections[SectionIndex];
-										int32 BufferVertIndex = Section.GetVertexBufferIndex() + VertIndex;
-										for (int32 InfluenceIndex = 0; InfluenceIndex < MaxBoneInfluences; InfluenceIndex++)
+									check(SectionIndex < LODRenderData->RenderSections.Num());
+									const FSkelMeshRenderSection& Section = LODRenderData->RenderSections[SectionIndex];
+									int32 BufferVertIndex = Section.GetVertexBufferIndex() + VertIndex;
+									for (int32 InfluenceIndex = 0; InfluenceIndex < MaxBoneInfluences; InfluenceIndex++)
+									{
+										const int32 BoneIndex = Section.BoneMap[SkinWeightVertexBuffer->GetBoneIndex(BufferVertIndex, InfluenceIndex)];
+										const float	Weight = (float)SkinWeightVertexBuffer->GetBoneWeight(BufferVertIndex, InfluenceIndex) * UE::AnimationCore::InvMaxRawBoneWeightFloat;
+										if (Weight > float(0) && 0 <= BoneIndex && BoneIndex < ComponentPose.Num())
 										{
-											const int32 BoneIndex = Section.BoneMap[SkinWeightVertexBuffer->GetBoneIndex(BufferVertIndex, InfluenceIndex)];
-											const float	Weight = (float)SkinWeightVertexBuffer->GetBoneWeight(BufferVertIndex, InfluenceIndex) * UE::AnimationCore::InvMaxRawBoneWeightFloat;
-											if (Weight > float(0) && 0 <= BoneIndex && BoneIndex < ComponentPose.Num())
-											{
-												BoneBoundVerts[BoneIndex].Add(IndicesMap[i]);
-												//BoneBoundWeights[BoneIndex].Add(Weight);
-												BoneBoundWeights[BoneIndex].Add(1.f);
-												break;
-											}
+											BoneBoundVerts[BoneIndex].Add(IndicesMap[i]);
+											//Only rigid skinning for now
+											BoneBoundWeights[BoneIndex].Add(1.f);
+											break;
 										}
 									}
 								}
 							}
 						}
-						else
-						{
-							TMap<int32, int32> BoneVertexToCollectionMap;
-							GeometryCollection::Facades::FPositionTargetFacade PositionTargets(InCollection);
-							PositionTargets.DefineSchema();
+					}
+					else
+					{
+						TMap<int32, int32> BoneVertexToCollectionMap;
+						GeometryCollection::Facades::FPositionTargetFacade PositionTargets(InCollection);
+						PositionTargets.DefineSchema();
 
-							for (int32 i = 0; i < Indices->Num(); i++)
+						for (int32 i = 0; i < Indices->Num(); i++)
+						{
+							TArray<int32> TriangleIntersections0 = VertexBVH.FindAllIntersections((*Vertices)[(*Indices)[i][0]]);
+							TArray<int32> TriangleIntersections1 = VertexBVH.FindAllIntersections((*Vertices)[(*Indices)[i][1]]);
+							TArray<int32> TriangleIntersections2 = VertexBVH.FindAllIntersections((*Vertices)[(*Indices)[i][2]]);
+							TriangleIntersections0.Sort();
+							TriangleIntersections1.Sort();
+							TriangleIntersections2.Sort();
+							TArray<int32> TriangleIntersections({});
+							for (int32 k = 0; k < TriangleIntersections0.Num(); k++)
 							{
-								TArray<int32> TriangleIntersections0 = VertexBVH.FindAllIntersections((*Vertices)[(*Indices)[i][0]]);
-								TArray<int32> TriangleIntersections1 = VertexBVH.FindAllIntersections((*Vertices)[(*Indices)[i][1]]);
-								TArray<int32> TriangleIntersections2 = VertexBVH.FindAllIntersections((*Vertices)[(*Indices)[i][2]]);
-								TriangleIntersections0.Sort();
-								TriangleIntersections1.Sort();
-								TriangleIntersections2.Sort();
-								TArray<int32> TriangleIntersections({});
-								for (int32 k = 0; k < TriangleIntersections0.Num(); k++)
+								if (TriangleIntersections1.Contains(TriangleIntersections0[k])
+									&& TriangleIntersections2.Contains(TriangleIntersections0[k]))
 								{
-									if (TriangleIntersections1.Contains(TriangleIntersections0[k])
-										&& TriangleIntersections2.Contains(TriangleIntersections0[k]))
-									{
-										TriangleIntersections.Emplace(TriangleIntersections0[k]);
-									}
+									TriangleIntersections.Emplace(TriangleIntersections0[k]);
 								}
-								int32 MinIndex = -1;
-								float MinDis = SphereRadius;
-								Chaos::TVector<float, 3> ClosestBary(0.f);
-								for (int32 j = 0; j < TriangleIntersections.Num(); j++)
+							}
+							int32 MinIndex = -1;
+							float MinDis = SphereRadius;
+							Chaos::TVector<float, 3> ClosestBary(0.f);
+							for (int32 j = 0; j < TriangleIntersections.Num(); j++)
+							{
+								Chaos::TVector<float, 3> Bary,
+									TriPos0((*Vertices)[(*Indices)[i][0]]), TriPos1((*Vertices)[(*Indices)[i][1]]), TriPos2((*Vertices)[(*Indices)[i][2]]),
+									ParticlePos(PositionVertexBuffer.VertexPosition(TriangleIntersections[j]));
+								Chaos::TVector<Chaos::FRealSingle, 3> ClosestPoint = Chaos::FindClosestPointAndBaryOnTriangle(TriPos0, TriPos1, TriPos2, ParticlePos, Bary);
+								Chaos::FRealSingle CurrentDistance = (ParticlePos - ClosestPoint).Size();
+								if (CurrentDistance < MinDis)
 								{
-									Chaos::TVector<float, 3> Bary,
-										TriPos0((*Vertices)[(*Indices)[i][0]]), TriPos1((*Vertices)[(*Indices)[i][1]]), TriPos2((*Vertices)[(*Indices)[i][2]]),
-										ParticlePos(PositionVertexBuffer.VertexPosition(TriangleIntersections[j]));
-									Chaos::TVector<Chaos::FRealSingle, 3> ClosestPoint = Chaos::FindClosestPointAndBaryOnTriangle(TriPos0, TriPos1, TriPos2, ParticlePos, Bary);
-									Chaos::FRealSingle CurrentDistance = (ParticlePos - ClosestPoint).Size();
-									if (CurrentDistance < MinDis)
-									{
-										MinDis = CurrentDistance;
-										MinIndex = TriangleIntersections[j];
-										ClosestBary = Bary;
-									}
+									MinDis = CurrentDistance;
+									MinIndex = TriangleIntersections[j];
+									ClosestBary = Bary;
 								}
+							}
 
-								if (MinIndex != -1)
+							if (MinIndex != -1)
+							{
+								int32 BoneParticleIndex = -1;
+								//TODO: add kinematic particles first
+								if (!BoneVertexToCollectionMap.Contains(MinIndex))
 								{
-									int32 BoneParticleIndex = -1;
-									//TODO: add kinematic particles first
-									if (!BoneVertexToCollectionMap.Contains(MinIndex))
-									{
 
-										int32 ParticleIndex = InCollection.AddElements(1, FGeometryCollection::VerticesGroup);
-										TManagedArray<FVector3f>& CurrentVertices = InCollection.ModifyAttribute<FVector3f>("Vertex", FGeometryCollection::VerticesGroup);
+									int32 ParticleIndex = InCollection.AddElements(1, FGeometryCollection::VerticesGroup);
+									TManagedArray<FVector3f>& CurrentVertices = InCollection.ModifyAttribute<FVector3f>("Vertex", FGeometryCollection::VerticesGroup);
 
-										CurrentVertices[ParticleIndex] = PositionVertexBuffer.VertexPosition(MinIndex);
-										BoneVertexToCollectionMap.Emplace(MinIndex, ParticleIndex);
+									CurrentVertices[ParticleIndex] = PositionVertexBuffer.VertexPosition(MinIndex);
+									BoneVertexToCollectionMap.Emplace(MinIndex, ParticleIndex);
 
-										int32 SectionIndex;
-										int32 VertIndex;
-										LODRenderData->GetSectionFromVertexIndex(MinIndex, SectionIndex, VertIndex);
+									int32 SectionIndex;
+									int32 VertIndex;
+									LODRenderData->GetSectionFromVertexIndex(MinIndex, SectionIndex, VertIndex);
 
-										check(SectionIndex < LODRenderData->RenderSections.Num());
-										const FSkelMeshRenderSection& Section = LODRenderData->RenderSections[SectionIndex];
-										int32 BufferVertIndex = Section.GetVertexBufferIndex() + VertIndex;
-										for (int32 InfluenceIndex = 0; InfluenceIndex < MaxBoneInfluences; InfluenceIndex++)
-										{
-											const int32 BoneIndex = Section.BoneMap[SkinWeightVertexBuffer->GetBoneIndex(BufferVertIndex, InfluenceIndex)];
-											const float	Weight = (float)SkinWeightVertexBuffer->GetBoneWeight(BufferVertIndex, InfluenceIndex) * UE::AnimationCore::InvMaxRawBoneWeightFloat;
-											if (Weight > float(0) && 0 <= BoneIndex && BoneIndex < ComponentPose.Num())
-											{
-												BoneBoundVerts[BoneIndex].Add(ParticleIndex);
-												BoneBoundWeights[BoneIndex].Add(Weight);
-											}
-										}
-									}
-
-									BoneParticleIndex = BoneVertexToCollectionMap[MinIndex];
-
-									GeometryCollection::Facades::FPositionTargetsData DataPackage;
-									DataPackage.TargetIndex.Init(BoneParticleIndex, 1);
-									DataPackage.TargetWeights.Init(1.f, 1);
-									DataPackage.SourceWeights.Init(1.f, 3);
-									DataPackage.SourceIndex.Init(-1, 3);
-									DataPackage.SourceIndex[0] = (*Indices)[i][0];
-									DataPackage.SourceIndex[1] = (*Indices)[i][1];
-									DataPackage.SourceIndex[2] = (*Indices)[i][2];
-									DataPackage.SourceWeights[0] = ClosestBary[0];
-									DataPackage.SourceWeights[1] = ClosestBary[1];
-									DataPackage.SourceWeights[2] = ClosestBary[2];
-									if (TManagedArray<float>* Mass = InCollection.FindAttribute<float>("Mass", FGeometryCollection::VerticesGroup))
+									check(SectionIndex < LODRenderData->RenderSections.Num());
+									const FSkelMeshRenderSection& Section = LODRenderData->RenderSections[SectionIndex];
+									int32 BufferVertIndex = Section.GetVertexBufferIndex() + VertIndex;
+									for (int32 InfluenceIndex = 0; InfluenceIndex < MaxBoneInfluences; InfluenceIndex++)
 									{
-										DataPackage.Stiffness = 0.f;
-										for (int32 k = 0; k < 3; k++)
+										const int32 BoneIndex = Section.BoneMap[SkinWeightVertexBuffer->GetBoneIndex(BufferVertIndex, InfluenceIndex)];
+										const float	Weight = (float)SkinWeightVertexBuffer->GetBoneWeight(BufferVertIndex, InfluenceIndex) * UE::AnimationCore::InvMaxRawBoneWeightFloat;
+										if (Weight > float(0) && 0 <= BoneIndex && BoneIndex < ComponentPose.Num())
 										{
-											DataPackage.Stiffness += DataPackage.SourceWeights[k] * PositionTargetStiffness * (*Mass)[DataPackage.SourceIndex[k]];
+											BoneBoundVerts[BoneIndex].Add(ParticleIndex);
+											BoneBoundWeights[BoneIndex].Add(Weight);
 										}
-										//DataPackage.Stiffness += DataPackage.TargetWeights[0] * PositionTargetStiffness * (*Mass)[DataPackage.TargetIndex[0]];
-										//DataPackage.Stiffness /= 2.f;
 									}
-									else
+								}
+
+								BoneParticleIndex = BoneVertexToCollectionMap[MinIndex];
+
+								GeometryCollection::Facades::FPositionTargetsData DataPackage;
+								DataPackage.TargetIndex.Init(BoneParticleIndex, 1);
+								DataPackage.TargetWeights.Init(1.f, 1);
+								DataPackage.SourceWeights.Init(1.f, 3);
+								DataPackage.SourceIndex.Init(-1, 3);
+								DataPackage.SourceIndex[0] = (*Indices)[i][0];
+								DataPackage.SourceIndex[1] = (*Indices)[i][1];
+								DataPackage.SourceIndex[2] = (*Indices)[i][2];
+								DataPackage.SourceWeights[0] = ClosestBary[0];
+								DataPackage.SourceWeights[1] = ClosestBary[1];
+								DataPackage.SourceWeights[2] = ClosestBary[2];
+								if (const TManagedArray<float>* Mass = InCollection.FindAttribute<float>("Mass", FGeometryCollection::VerticesGroup))
+								{
+									DataPackage.Stiffness = 0.f;
+									//Target is kinematic, only compute stiffness from source
+									for (int32 k = 0; k < 3; k++)
 									{
-										DataPackage.Stiffness = PositionTargetStiffness;
+										DataPackage.Stiffness += DataPackage.SourceWeights[k] * PositionTargetStiffness * (*Mass)[DataPackage.SourceIndex[k]];
 									}
-									PositionTargets.AddPositionTarget(DataPackage);
 								}
-
-
+								else
+								{
+									DataPackage.Stiffness = PositionTargetStiffness;
+								}
+								PositionTargets.AddPositionTarget(DataPackage);
 							}
 
 
 						}
-						auto DoubleVert = [](FVector3f V) { return FVector3d(V.X, V.Y, V.Z); };
 
-						for (int32 BoneIndex = 0; BoneIndex < ComponentPose.Num(); ++BoneIndex)
+
+					}
+					auto DoubleVert = [](FVector3f V) { return FVector3d(V.X, V.Y, V.Z); };
+
+					for (int32 BoneIndex = 0; BoneIndex < ComponentPose.Num(); ++BoneIndex)
+					{
+						if (BoneBoundVerts[BoneIndex].Num())
 						{
-							if (BoneBoundVerts[BoneIndex].Num())
+							FString BoneName = BoneSkeletalMesh->GetRefSkeleton().GetBoneName(BoneIndex).ToString();
+							//get local coords of bound verts
+							using FKinematics = GeometryCollection::Facades::FKinematicBindingFacade;
+							FKinematics Kinematics(InCollection); Kinematics.DefineSchema();
+							if (Kinematics.IsValid())
 							{
-								FString BoneName = BoneSkeletalMesh->GetRefSkeleton().GetBoneName(BoneIndex).ToString();
-								//get local coords of bound verts
-								typedef GeometryCollection::Facades::FKinematicBindingFacade FKinematics;
-								FKinematics Kinematics(InCollection); Kinematics.DefineSchema();
-								if (Kinematics.IsValid())
+								if (ensure(BoneNameIndexMap.Contains(BoneName)))
 								{
-									FKinematics::FBindingKey Binding = Kinematics.SetBoneBindings(BoneIndex, BoneBoundVerts[BoneIndex], BoneBoundWeights[BoneIndex]);
+									FKinematics::FBindingKey Binding = Kinematics.SetBoneBindings(BoneNameIndexMap[BoneName], BoneBoundVerts[BoneIndex], BoneBoundWeights[BoneIndex]);
 									TManagedArray<TArray<FVector3f>>& LocalPos = InCollection.AddAttribute<TArray<FVector3f>>("LocalPosition", Binding.GroupName);
 									Kinematics.AddKinematicBinding(Binding);
 
@@ -335,10 +319,8 @@ void FSetFleshBonePositionTargetBindingDataflowNode::Evaluate(Dataflow::FContext
 								}
 							}
 						}
-						GeometryCollection::Facades::FVertexBoneWeightsFacade(InCollection).AddBoneWeightsFromKinematicBindings();
-
 					}
-
+					GeometryCollection::Facades::FVertexBoneWeightsFacade(InCollection).AddBoneWeightsFromKinematicBindings();
 
 				}
 			}
diff --git a/Engine/Plugins/Experimental/ChaosFlesh/Source/ChaosFleshNodes/Private/Dataflow/ChaosFleshSetVertexTrianglePositionTargetBindingNode.cpp b/Engine/Plugins/Experimental/ChaosFlesh/Source/ChaosFleshNodes/Private/Dataflow/ChaosFleshSetVertexTrianglePositionTargetBindingNode.cpp
--- a/Engine/Plugins/Experimental/ChaosFlesh/Source/ChaosFleshNodes/Private/Dataflow/ChaosFleshSetVertexTrianglePositionTargetBindingNode.cpp
+++ b/Engine/Plugins/Experimental/ChaosFlesh/Source/ChaosFleshNodes/Private/Dataflow/ChaosFleshSetVertexTrianglePositionTargetBindingNode.cpp
@@ -93,9 +93,11 @@ void FSetVertexTrianglePositionTargetBindingDataflowNode::Evaluate(Dataflow::FCo
 					Chaos::TVec3<float> CoordDiff = (CoordMaxs - CoordMins) * VertexRadiusRatio;
 					Chaos::FReal SphereRadius = Chaos::FReal(FGenericPlatformMath::Max(CoordDiff[0], FGenericPlatformMath::Max(CoordDiff[1], CoordDiff[2])));
 					
-					if (FindInput(&VertexSelectionSetIn) && FindInput(&VertexSelectionSetIn)->GetConnection())
+					if (IsConnected(&VertexSelection))
 					{
-						IndicesMap = GetValue<TArray<int32>>(Context, &VertexSelectionSetIn);
+						FDataflowVertexSelection InDataflowVertexSelection = GetValue<FDataflowVertexSelection>(Context, &VertexSelection);
+						IndicesMap = InDataflowVertexSelection.AsArray();
+
 					}
 
 					GeometryCollection::Facades::FPositionTargetFacade PositionTargets(InCollection);
diff --git a/Engine/Plugins/Experimental/ChaosFlesh/Source/ChaosFleshNodes/Public/Dataflow/ChaosFleshGetSurfaceIndicesNode.h b/Engine/Plugins/Experimental/ChaosFlesh/Source/ChaosFleshNodes/Public/Dataflow/ChaosFleshGetSurfaceIndicesNode.h
--- a/Engine/Plugins/Experimental/ChaosFlesh/Source/ChaosFleshNodes/Public/Dataflow/ChaosFleshGetSurfaceIndicesNode.h
+++ b/Engine/Plugins/Experimental/ChaosFlesh/Source/ChaosFleshNodes/Public/Dataflow/ChaosFleshGetSurfaceIndicesNode.h
@@ -5,7 +5,7 @@
 #include "CoreMinimal.h"
 #include "Dataflow/DataflowCore.h"
 #include "Dataflow/DataflowEngine.h"
-
+#include "Dataflow/DataflowSelection.h"
 #include "ChaosFleshGetSurfaceIndicesNode.generated.h"
 
 USTRUCT(meta = (DataflowFlesh))
@@ -22,15 +22,15 @@ public:
 	UPROPERTY(meta = (DataflowInput, DisplayName = "GeometryGroupGuidsIn"))
 	TArray<FString> GeometryGroupGuidsIn;
 
-	UPROPERTY(meta = (DataflowOutput, DisplayName = "SurfaceIndicesOut"))
-	TArray<int32> SurfaceIndicesOut;
+	UPROPERTY(meta = (DataflowOutput, DisplayName = "SurfaceVertexSelection"))
+	FDataflowVertexSelection SurfaceVertexSelection;
 
 	FGetSurfaceIndicesNode(const Dataflow::FNodeParameters& InParam, FGuid InGuid = FGuid::NewGuid())
 		: FDataflowNode(InParam, InGuid)
 	{
 		RegisterInputConnection(&Collection);
 		RegisterInputConnection(&GeometryGroupGuidsIn);
-		RegisterOutputConnection(&SurfaceIndicesOut);
+		RegisterOutputConnection(&SurfaceVertexSelection);
 	}
 
 	virtual void Evaluate(Dataflow::FContext& Context, const FDataflowOutput* Out) const override;
diff --git a/Engine/Plugins/Experimental/ChaosFlesh/Source/ChaosFleshNodes/Public/Dataflow/ChaosFleshSetVertexTrianglePositionTargetBindingNode.h b/Engine/Plugins/Experimental/ChaosFlesh/Source/ChaosFleshNodes/Public/Dataflow/ChaosFleshSetVertexTrianglePositionTargetBindingNode.h
--- a/Engine/Plugins/Experimental/ChaosFlesh/Source/ChaosFleshNodes/Public/Dataflow/ChaosFleshSetVertexTrianglePositionTargetBindingNode.h
+++ b/Engine/Plugins/Experimental/ChaosFlesh/Source/ChaosFleshNodes/Public/Dataflow/ChaosFleshSetVertexTrianglePositionTargetBindingNode.h
@@ -5,6 +5,7 @@
 #include "CoreMinimal.h"
 #include "Dataflow/DataflowCore.h"
 #include "Dataflow/DataflowEngine.h"
+#include "Dataflow/DataflowSelection.h"
 
 #include "ChaosFleshSetVertexTrianglePositionTargetBindingNode.generated.h"
 
@@ -25,8 +26,8 @@ public:
 	UPROPERTY(EditAnywhere, Category = "Dataflow")
 		float PositionTargetStiffness = 10000.f;
 
-	UPROPERTY(meta = (DataflowInput, DisplayName = "VertexSelectionSet"))
-		TArray<int32> VertexSelectionSetIn;
+	UPROPERTY(meta = (DataflowInput, DisplayName = "(Optional) VertexSelection"))
+		FDataflowVertexSelection VertexSelection;
 
 	UPROPERTY(EditAnywhere, Category = "Dataflow", meta = (ClampMin = "0.0", ClampMax = "2.0", UIMin = "0.0", UIMax = "2.0"))
 		float VertexRadiusRatio = .001f;
@@ -36,7 +37,7 @@ public:
 		: FDataflowNode(InParam, InGuid)
 	{
 		RegisterInputConnection(&Collection);
-		RegisterInputConnection(&VertexSelectionSetIn);
+		RegisterInputConnection(&VertexSelection);
 		RegisterOutputConnection(&Collection, &Collection);
 	}
 
diff --git a/Engine/Plugins/Experimental/Dataflow/Source/DataflowEditor/Private/DataflowEditorTools/DataflowEditorWeightMapPaintTool.cpp b/Engine/Plugins/Experimental/Dataflow/Source/DataflowEditor/Private/DataflowEditorTools/DataflowEditorWeightMapPaintTool.cpp
--- a/Engine/Plugins/Experimental/Dataflow/Source/DataflowEditor/Private/DataflowEditorTools/DataflowEditorWeightMapPaintTool.cpp
+++ b/Engine/Plugins/Experimental/Dataflow/Source/DataflowEditor/Private/DataflowEditorTools/DataflowEditorWeightMapPaintTool.cpp
@@ -1025,11 +1025,11 @@ void UDataflowEditorWeightMapPaintTool::OnPolyLassoFinished(const FCameraPolyLas
 	// project each mesh vertex to view plane and evaluate winding integral of polyline
 	const FDynamicMesh3* Mesh = GetSculptMesh();
 	TempROIBuffer.SetNum(Mesh->MaxVertexID());
-	ParallelFor(Mesh->MaxVertexID(), [&](int32 vid)
+	ParallelFor(Mesh->MaxVertexID(), [&](int32 VertexIdx)
 	{
-		if (Mesh->IsVertex(vid))
+		if (Mesh->IsVertex(VertexIdx))
 		{
-			FVector3d WorldPos = CurTargetTransform.TransformPosition(Mesh->GetVertex(vid));
+			FVector3d WorldPos = CurTargetTransform.TransformPosition(Mesh->GetVertex(VertexIdx));
 			FVector2f PlanePos = (FVector2f)Lasso.GetProjectedPoint((FVector)WorldPos);
 
 			double WindingSum = 0;
@@ -1042,17 +1042,17 @@ void UDataflowEditorWeightMapPaintTool::OnPolyLassoFinished(const FCameraPolyLas
 			}
 			WindingSum /= FMathd::TwoPi;
 			bool bInside = FMathd::Abs(WindingSum) > 0.3;
-			TempROIBuffer[vid] = bInside ? 1 : 0;
+			TempROIBuffer[VertexIdx] = bInside ? 1 : 0;
 		}
 		else
 		{
-			TempROIBuffer[vid] = -1;
+			TempROIBuffer[VertexIdx] = -1;
 		}
 	});
 
 	// convert to vertex selection, and then select fully-enclosed faces
 	FMeshVertexSelection VertexSelection(Mesh);
-	VertexSelection.SelectByVertexID([&](int32 vid) { return TempROIBuffer[vid] == 1; });
+	VertexSelection.SelectByVertexID([&](int32 VertexIdx) { return TempROIBuffer[VertexIdx] == 1; });
 
 	double SetWeightValue = GetInEraseStroke() ? 0.0 : FilterProperties->AttributeValue;
 	SetVerticesToWeightMap(VertexSelection.AsSet(), SetWeightValue, GetInEraseStroke());
@@ -1071,16 +1071,16 @@ void UDataflowEditorWeightMapPaintTool::ComputeGradient()
 
 	const FDynamicMesh3* const Mesh = DynamicMeshComponent->GetMesh();
 	TempROIBuffer.SetNum(0, EAllowShrinking::No);
-	for (int32 vid : Mesh->VertexIndicesItr())
+	for (int32 VertexIdx : Mesh->VertexIndicesItr())
 	{
-		TempROIBuffer.Add(vid);
+		TempROIBuffer.Add(VertexIdx);
 	}
 
 	if (bHaveDynamicMeshToWeightConversion)
 	{
-		for (int32 vid : TempROIBuffer)
+		for (int32 VertexIdx : TempROIBuffer)
 		{
-			for (const int32 Idx : WeightToDynamicMesh[DynamicMeshToWeight[vid]])
+			for (const int32 Idx : WeightToDynamicMesh[DynamicMeshToWeight[VertexIdx]])
 			{
 				ActiveWeightEditChangeTracker->SaveVertexOneRingTriangles(Idx, true);
 			}
@@ -1088,9 +1088,9 @@ void UDataflowEditorWeightMapPaintTool::ComputeGradient()
 	}
 	else
 	{
-		for (int32 vid : TempROIBuffer)
+		for (int32 VertexIdx : TempROIBuffer)
 		{
-			ActiveWeightEditChangeTracker->SaveVertexOneRingTriangles(vid, true);
+			ActiveWeightEditChangeTracker->SaveVertexOneRingTriangles(VertexIdx, true);
 		}
 	}
 
@@ -1181,9 +1181,9 @@ void UDataflowEditorWeightMapPaintTool::SetVerticesToWeightMap(const TSet<int32>
 	BeginChange();
 
 	TempROIBuffer.SetNum(0, EAllowShrinking::No);
-	for (int32 vid : Vertices)
+	for (int32 VertexIdx : Vertices)
 	{
-		TempROIBuffer.Add(vid);
+		TempROIBuffer.Add(VertexIdx);
 	}
 
 	if (HaveVisibilityFilter())
@@ -1196,9 +1196,9 @@ void UDataflowEditorWeightMapPaintTool::SetVerticesToWeightMap(const TSet<int32>
 
 	if (bHaveDynamicMeshToWeightConversion)
 	{
-		for (int32 vid : TempROIBuffer)
+		for (int32 VertexIdx : TempROIBuffer)
 		{
-			for (const int32 Idx : WeightToDynamicMesh[DynamicMeshToWeight[vid]])
+			for (const int32 Idx : WeightToDynamicMesh[DynamicMeshToWeight[VertexIdx]])
 			{
 				ActiveWeightEditChangeTracker->SaveVertexOneRingTriangles(Idx, true);
 				ActiveWeightMap->SetValue(Idx, &WeightValue);
@@ -1207,13 +1207,13 @@ void UDataflowEditorWeightMapPaintTool::SetVerticesToWeightMap(const TSet<int32>
 	}
 	else
 	{
-		for (int32 vid : TempROIBuffer)
+		for (int32 VertexIdx : TempROIBuffer)
 		{
-			ActiveWeightEditChangeTracker->SaveVertexOneRingTriangles(vid, true);
+			ActiveWeightEditChangeTracker->SaveVertexOneRingTriangles(VertexIdx, true);
 		}
-		for (int32 vid : TempROIBuffer)
+		for (int32 VertexIdx : TempROIBuffer)
 		{		
-			ActiveWeightMap->SetValue(vid, &WeightValue);
+			ActiveWeightMap->SetValue(VertexIdx, &WeightValue);
 		}
 	}
 
@@ -1239,18 +1239,18 @@ void UDataflowEditorWeightMapPaintTool::ApplyVisibilityFilter(TSet<int32>& Verti
 {
 	ROIBuffer.SetNum(0, EAllowShrinking::No);
 	ROIBuffer.Reserve(Vertices.Num());
-	for (int32 vid : Vertices)
+	for (int32 VertexIdx : Vertices)
 	{
-		ROIBuffer.Add(vid);
+		ROIBuffer.Add(VertexIdx);
 	}
 	
 	OutputBuffer.Reset();
 	ApplyVisibilityFilter(TempROIBuffer, OutputBuffer);
 
 	Vertices.Reset();
-	for (int32 vid : OutputBuffer)
+	for (int32 VertexIdx : OutputBuffer)
 	{
-		Vertices.Add(vid);
+		Vertices.Add(VertexIdx);
 	}
 }
 
@@ -1611,16 +1611,16 @@ void UDataflowEditorWeightMapPaintTool::FloodFillCurrentWeightAction()
 	const float SetWeightValue = FilterProperties->AttributeValue;
 	const FDynamicMesh3* Mesh = DynamicMeshComponent->GetMesh();
 	TempROIBuffer.SetNum(0, EAllowShrinking::No);
-	for (int32 vid : Mesh->VertexIndicesItr())
+	for (int32 VertexIdx : Mesh->VertexIndicesItr())
 	{
-		TempROIBuffer.Add(vid);
+		TempROIBuffer.Add(VertexIdx);
 	}
 
 	if (bHaveDynamicMeshToWeightConversion)
 	{
-		for (int32 vid : TempROIBuffer)
+		for (int32 VertexIdx : TempROIBuffer)
 		{
-			for (const int32 Idx : WeightToDynamicMesh[DynamicMeshToWeight[vid]])
+			for (const int32 Idx : WeightToDynamicMesh[DynamicMeshToWeight[VertexIdx]])
 			{
 				ActiveWeightEditChangeTracker->SaveVertexOneRingTriangles(Idx, true);
 				ActiveWeightMap->SetValue(Idx, &SetWeightValue);
@@ -1629,13 +1629,13 @@ void UDataflowEditorWeightMapPaintTool::FloodFillCurrentWeightAction()
 	}
 	else
 	{
-		for (int32 vid : TempROIBuffer)
+		for (int32 VertexIdx : TempROIBuffer)
 		{
-			ActiveWeightEditChangeTracker->SaveVertexOneRingTriangles(vid, true);
+			ActiveWeightEditChangeTracker->SaveVertexOneRingTriangles(VertexIdx, true);
 		}
-		for (int32 vid : TempROIBuffer)
+		for (int32 VertexIdx : TempROIBuffer)
 		{
-			ActiveWeightMap->SetValue(vid, &SetWeightValue);
+			ActiveWeightMap->SetValue(VertexIdx, &SetWeightValue);
 		}
 	}
 
@@ -1646,7 +1646,6 @@ void UDataflowEditorWeightMapPaintTool::FloodFillCurrentWeightAction()
 	EndChange();
 }
 
-
 void UDataflowEditorWeightMapPaintTool::ClearAllWeightsAction()
 {
 	if (!ActiveWeightMap)
@@ -1659,16 +1658,16 @@ void UDataflowEditorWeightMapPaintTool::ClearAllWeightsAction()
 	float SetWeightValue = 0.0f;
 	const FDynamicMesh3* Mesh = DynamicMeshComponent->GetMesh();
 	TempROIBuffer.SetNum(0, EAllowShrinking::No);
-	for (int32 vid : Mesh->VertexIndicesItr())
+	for (int32 VertexIdx : Mesh->VertexIndicesItr())
 	{
-		TempROIBuffer.Add(vid);
+		TempROIBuffer.Add(VertexIdx);
 	}
 
 	if (bHaveDynamicMeshToWeightConversion)
 	{
-		for (int32 vid : TempROIBuffer)
+		for (int32 VertexIdx : TempROIBuffer)
 		{
-			for (const int32 Idx : WeightToDynamicMesh[DynamicMeshToWeight[vid]])
+			for (const int32 Idx : WeightToDynamicMesh[DynamicMeshToWeight[VertexIdx]])
 			{
 				ActiveWeightEditChangeTracker->SaveVertexOneRingTriangles(Idx, true);
 				ActiveWeightMap->SetValue(Idx, &SetWeightValue);
@@ -1677,13 +1676,13 @@ void UDataflowEditorWeightMapPaintTool::ClearAllWeightsAction()
 	}
 	else
 	{
-		for (int32 vid : TempROIBuffer)
+		for (int32 VertexIdx : TempROIBuffer)
 		{
-			ActiveWeightEditChangeTracker->SaveVertexOneRingTriangles(vid, true);
+			ActiveWeightEditChangeTracker->SaveVertexOneRingTriangles(VertexIdx, true);
 		}
-		for (int32 vid : TempROIBuffer)
+		for (int32 VertexIdx : TempROIBuffer)
 		{
-			ActiveWeightMap->SetValue(vid, &SetWeightValue);
+			ActiveWeightMap->SetValue(VertexIdx, &SetWeightValue);
 		}
 	}
 
@@ -1694,6 +1693,61 @@ void UDataflowEditorWeightMapPaintTool::ClearAllWeightsAction()
 	EndChange();
 }
 
+void UDataflowEditorWeightMapPaintTool::InvertCurrentWeightAction(bool bInvertSurfaceOnly)
+{
+	if (!ActiveWeightMap)
+	{
+		return;
+	}
+
+	BeginChange();
+
+
+	const FDynamicMesh3* Mesh = DynamicMeshComponent->GetMesh();
+	TempROIBuffer.SetNum(0, EAllowShrinking::No);
+	for (int32 VertexIdx : Mesh->VertexIndicesItr())
+	{
+		if (!bInvertSurfaceOnly || Mesh->IsReferencedVertex(VertexIdx))
+		{
+			TempROIBuffer.Add(VertexIdx);
+		}
+	}
+
+	if (bHaveDynamicMeshToWeightConversion)
+	{
+		for (int32 VertexIdx : TempROIBuffer)
+		{
+			for (const int32 Idx : WeightToDynamicMesh[DynamicMeshToWeight[VertexIdx]])
+			{
+				ActiveWeightEditChangeTracker->SaveVertexOneRingTriangles(Idx, true);
+				float CurrentValue;
+				ActiveWeightMap->GetValue(Idx, &CurrentValue);
+				float SetWeightValue = 1.f - CurrentValue;
+				ActiveWeightMap->SetValue(Idx, &SetWeightValue);
+			}
+		}
+	}
+	else
+	{
+		for (int32 VertexIdx : TempROIBuffer)
+		{
+			ActiveWeightEditChangeTracker->SaveVertexOneRingTriangles(VertexIdx, true);
+		}
+		for (int32 VertexIdx : TempROIBuffer)
+		{
+			float CurrentValue;
+			ActiveWeightMap->GetValue(VertexIdx, &CurrentValue);
+			float SetWeightValue = 1.f - CurrentValue;
+			ActiveWeightMap->SetValue(VertexIdx, &SetWeightValue);
+		}
+	}
+
+	// update colors
+	UpdateVertexColorOverlay();
+	DynamicMeshComponent->FastNotifyVertexAttributesUpdated(EMeshRenderAttributeFlags::VertexColors);
+	GetToolManager()->PostInvalidation();
+	EndChange();
+}
 
 void UDataflowEditorWeightMapPaintTool::UpdateSelectedNode()
 {
@@ -1993,6 +2047,14 @@ void UDataflowEditorWeightMapPaintTool::ApplyAction(EDataflowEditorWeightMapPain
 	case EDataflowEditorWeightMapPaintToolActions::ClearAll:
 		ClearAllWeightsAction();
 		break;
+
+	case EDataflowEditorWeightMapPaintToolActions::InvertCurrent:
+		InvertCurrentWeightAction(false);
+		break;
+
+	case EDataflowEditorWeightMapPaintToolActions::InvertCurrentSurface:
+		InvertCurrentWeightAction(true);
+		break;
 	}
 }
 
diff --git a/Engine/Plugins/Experimental/Dataflow/Source/DataflowEditor/Public/DataflowEditorTools/DataflowEditorWeightMapPaintTool.h b/Engine/Plugins/Experimental/Dataflow/Source/DataflowEditor/Public/DataflowEditorTools/DataflowEditorWeightMapPaintTool.h
--- a/Engine/Plugins/Experimental/Dataflow/Source/DataflowEditor/Public/DataflowEditorTools/DataflowEditorWeightMapPaintTool.h
+++ b/Engine/Plugins/Experimental/Dataflow/Source/DataflowEditor/Public/DataflowEditorTools/DataflowEditorWeightMapPaintTool.h
@@ -206,9 +206,10 @@ UENUM()
 enum class EDataflowEditorWeightMapPaintToolActions
 {
 	NoAction,
-
 	FloodFillCurrent,
 	ClearAll,
+	InvertCurrent,
+	InvertCurrentSurface
 };
 
 
@@ -237,6 +238,19 @@ public:
 		PostAction(EDataflowEditorWeightMapPaintToolActions::FloodFillCurrent);
 	}
 
+	/* Invert the values in range [0, 1] for the current selected geometry, including interior vertices*/
+	UFUNCTION(CallInEditor, Category = Operations, meta = (DisplayPriority = 13))
+	void InvertCurrent()
+	{
+		PostAction(EDataflowEditorWeightMapPaintToolActions::InvertCurrent);
+	}
+
+	/* Invert the values in range [0, 1] for the current selected surface*/
+	UFUNCTION(CallInEditor, Category = Operations, meta = (DisplayPriority = 13))
+	void InvertCurrentSurface()
+	{
+		PostAction(EDataflowEditorWeightMapPaintToolActions::InvertCurrentSurface);
+	}
 };
 
 UCLASS()
@@ -315,6 +329,7 @@ private:
 public:
 	void FloodFillCurrentWeightAction();
 	void ClearAllWeightsAction();
+	void InvertCurrentWeightAction(bool bInvertSurfaceOnly = true);
 
 	void SetVerticesToWeightMap(const TSet<int32>& Vertices, double WeightValue, bool bIsErase);
 
diff --git a/Engine/Plugins/Experimental/Dataflow/Source/DataflowNodes/Public/Dataflow/DataflowCollectionAddScalarVertexPropertyNode.h b/Engine/Plugins/Experimental/Dataflow/Source/DataflowNodes/Public/Dataflow/DataflowCollectionAddScalarVertexPropertyNode.h
--- a/Engine/Plugins/Experimental/Dataflow/Source/DataflowNodes/Public/Dataflow/DataflowCollectionAddScalarVertexPropertyNode.h
+++ b/Engine/Plugins/Experimental/Dataflow/Source/DataflowNodes/Public/Dataflow/DataflowCollectionAddScalarVertexPropertyNode.h
@@ -49,7 +49,7 @@ struct FScalarVertexPropertyGroup
 	GENERATED_USTRUCT_BODY()
 
 	UPROPERTY(EditAnywhere, Category = "Vertex Group")
-	FName Name;
+	FName Name = FGeometryCollection::VerticesGroup;
 };
 
 
diff --git a/Engine/Plugins/Experimental/GeometryCollectionPlugin/Source/GeometryCollectionNodes/Private/Dataflow/GeometryCollectionSelectionNodes.cpp b/Engine/Plugins/Experimental/GeometryCollectionPlugin/Source/GeometryCollectionNodes/Private/Dataflow/GeometryCollectionSelectionNodes.cpp
--- a/Engine/Plugins/Experimental/GeometryCollectionPlugin/Source/GeometryCollectionNodes/Private/Dataflow/GeometryCollectionSelectionNodes.cpp
+++ b/Engine/Plugins/Experimental/GeometryCollectionPlugin/Source/GeometryCollectionNodes/Private/Dataflow/GeometryCollectionSelectionNodes.cpp
@@ -1280,11 +1280,18 @@ void FCollectionSelectionByAttrDataflowNode::Evaluate(Dataflow::FContext& Contex
 		Out->IsA<FDataflowMaterialSelection>(&MaterialSelection))
 	{
 		const FManagedArrayCollection& InCollection = GetValue<FManagedArrayCollection>(Context, &Collection);
-		const FName GroupName = Dataflow::Private::GetAttributeFromEnumAsName(Group);
+		FCollectionAttributeKey InAttributeKey = GetValue<FCollectionAttributeKey>(Context, &AttributeKey);
+		FName GroupName = Dataflow::Private::GetAttributeFromEnumAsName(Group);
+		FName AttributeName = FName(Attribute);
+		if (IsConnected(&AttributeKey))
+		{
+			GroupName = FName(InAttributeKey.Group);
+			AttributeName = FName(InAttributeKey.Attribute);
+		}
 
 		if (InCollection.HasGroup(GroupName))
 		{
-			if (InCollection.HasAttribute(FName(*Attribute), GroupName))
+			if (InCollection.HasAttribute(AttributeName, GroupName))
 			{
 				const int32 NumFaces = InCollection.NumElements(GroupName);
 
@@ -1293,16 +1300,16 @@ void FCollectionSelectionByAttrDataflowNode::Evaluate(Dataflow::FContext& Contex
 
 				CreateSelectionFromAttr(InCollection,
 					GroupName,
-					FName(*Attribute),
+					AttributeName,
 					Value,
 					Operation,
 					NewSelection);
 
-				SetValue(Context, Group == ESelectionByAttrGroup::Vertices ? MoveTemp(NewSelection) : FDataflowSelection(), &VertexSelection);
-				SetValue(Context, Group == ESelectionByAttrGroup::Faces ? MoveTemp(NewSelection) : FDataflowSelection(), &FaceSelection);
-				SetValue(Context, Group == ESelectionByAttrGroup::Transform ? MoveTemp(NewSelection) : FDataflowSelection(), &TransformSelection);
-				SetValue(Context, Group == ESelectionByAttrGroup::Geometry ? MoveTemp(NewSelection) : FDataflowSelection(), &GeometrySelection);
-				SetValue(Context, Group == ESelectionByAttrGroup::Material ? MoveTemp(NewSelection) : FDataflowSelection(), &MaterialSelection);
+				SetValue(Context, GroupName == Dataflow::Private::GetAttributeFromEnumAsName(ESelectionByAttrGroup::Vertices) ? MoveTemp(NewSelection) : FDataflowSelection(), &VertexSelection);
+				SetValue(Context, GroupName == Dataflow::Private::GetAttributeFromEnumAsName(ESelectionByAttrGroup::Faces) ? MoveTemp(NewSelection) : FDataflowSelection(), &FaceSelection);
+				SetValue(Context, GroupName == Dataflow::Private::GetAttributeFromEnumAsName(ESelectionByAttrGroup::Transform) ? MoveTemp(NewSelection) : FDataflowSelection(), &TransformSelection);
+				SetValue(Context, GroupName == Dataflow::Private::GetAttributeFromEnumAsName(ESelectionByAttrGroup::Geometry) ? MoveTemp(NewSelection) : FDataflowSelection(), &GeometrySelection);
+				SetValue(Context, GroupName == Dataflow::Private::GetAttributeFromEnumAsName(ESelectionByAttrGroup::Material) ? MoveTemp(NewSelection) : FDataflowSelection(), &MaterialSelection);
 
 				return;
 			}
diff --git a/Engine/Plugins/Experimental/GeometryCollectionPlugin/Source/GeometryCollectionNodes/Public/Dataflow/GeometryCollectionSelectionNodes.h b/Engine/Plugins/Experimental/GeometryCollectionPlugin/Source/GeometryCollectionNodes/Public/Dataflow/GeometryCollectionSelectionNodes.h
--- a/Engine/Plugins/Experimental/GeometryCollectionPlugin/Source/GeometryCollectionNodes/Public/Dataflow/GeometryCollectionSelectionNodes.h
+++ b/Engine/Plugins/Experimental/GeometryCollectionPlugin/Source/GeometryCollectionNodes/Public/Dataflow/GeometryCollectionSelectionNodes.h
@@ -8,8 +8,9 @@
 #include "Dataflow/DataflowSelection.h"
 #include "Math/MathFwd.h"
 #include "Math/Sphere.h"
-#include "GeometryCollectionSelectionNodes.generated.h"
+#include "Dataflow/DataflowConnectionTypes.h"
 
+#include "GeometryCollectionSelectionNodes.generated.h"
 
 class FGeometryCollection;
 
@@ -1416,6 +1417,10 @@ public:
 	UPROPERTY(meta = (DataflowInput, DataflowOutput, DataflowPassthrough = "Collection", DataflowIntrinsic))
 	FManagedArrayCollection Collection;
 
+	/** AttributeKey input */
+	UPROPERTY(meta = (DataflowInput))
+	FCollectionAttributeKey AttributeKey;
+
 	/** Group */
 	UPROPERTY(EditAnywhere, Category = "Selection")
 	ESelectionByAttrGroup Group = ESelectionByAttrGroup::Faces;
@@ -1456,6 +1461,7 @@ public:
 		: FDataflowNode(InParam, InGuid)
 	{
 		RegisterInputConnection(&Collection);
+		RegisterInputConnection(&AttributeKey);
 		RegisterOutputConnection(&Collection, &Collection);
 		RegisterOutputConnection(&VertexSelection);
 		RegisterOutputConnection(&FaceSelection);
diff --git a/Engine/Source/Runtime/Experimental/Chaos/Private/GeometryCollection/Facades/CollectionVertexBoneWeightsFacade.cpp b/Engine/Source/Runtime/Experimental/Chaos/Private/GeometryCollection/Facades/CollectionVertexBoneWeightsFacade.cpp
--- a/Engine/Source/Runtime/Experimental/Chaos/Private/GeometryCollection/Facades/CollectionVertexBoneWeightsFacade.cpp
+++ b/Engine/Source/Runtime/Experimental/Chaos/Private/GeometryCollection/Facades/CollectionVertexBoneWeightsFacade.cpp
@@ -73,7 +73,31 @@ namespace GeometryCollection::Facades
 		}
 	}
 
-
+	void FVertexBoneWeightsFacade::ModifyBoneWeight(int32 VertexIndex, TArray<int32> VertexBoneIndex, TArray<float> VertexBoneWeight)
+	{
+		TManagedArray< TArray<int32> >& IndicesArray = BoneIndexAttribute.Modify();
+		TManagedArray< TArray<float> >& WeightsArray = BoneWeightAttribute.Modify();
+		const TManagedArray<FVector3f>& Vertices = VerticesAttribute.Modify();
+		if (VerticesAttribute.IsValidIndex(VertexIndex))
+		{
+			IndicesArray[VertexIndex].Empty();
+			WeightsArray[VertexIndex].Empty();
+			float TotalWeight = 0.f;
+			for (int32 Idx = 0; Idx < VertexBoneIndex.Num(); ++Idx)
+			{
+				if (ParentAttribute.IsValidIndex(VertexBoneIndex[Idx]))
+				{
+					IndicesArray[VertexIndex].Add(VertexBoneIndex[Idx]);
+					WeightsArray[VertexIndex].Add(VertexBoneWeight[Idx]);
+					TotalWeight += VertexBoneWeight[Idx];
+				}
+			}
+			if (TotalWeight < 1.f - UE_KINDA_SMALL_NUMBER || TotalWeight > 1.f + UE_KINDA_SMALL_NUMBER)
+			{
+				UE_LOG(LogChaos, Warning, TEXT("FVertexBoneWeightsFacade::ModifyBoneWeight: Bone weight sum %f is not 1 on vertex %d"), TotalWeight, VertexIndex);
+			}
+		}
+	}
 	//
 	//  Add Weights from Selection 
 	//
@@ -119,7 +143,7 @@ namespace GeometryCollection::Facades
 						if (0 <= Vert && Vert < NumVertices && !IndicesArray[Vert].Contains(Bone))
 						{
 							int32 BoneIndex = IndicesArray[Vert].Find(Bone);
-							if (TotalWeights[Vert] + Weight <= 1.f)
+							if (TotalWeights[Vert] + Weight <= 1.f + UE_KINDA_SMALL_NUMBER)
 							{
 								if (BoneIndex == INDEX_NONE)
 								{
@@ -135,7 +159,7 @@ namespace GeometryCollection::Facades
 							}
 							else
 							{
-								UE_LOG(LogChaos, Warning, TEXT("Bone weight sum exceeds 1 on vertex %d"), Vert);
+								UE_LOG(LogChaos, Warning, TEXT("Bone weight sum %f exceeds 1 on vertex %d"), TotalWeights[Vert] + Weight, Vert);
 							}
 						}
 					}
diff --git a/Engine/Source/Runtime/Experimental/Chaos/Public/GeometryCollection/Facades/CollectionVertexBoneWeightsFacade.h b/Engine/Source/Runtime/Experimental/Chaos/Public/GeometryCollection/Facades/CollectionVertexBoneWeightsFacade.h
--- a/Engine/Source/Runtime/Experimental/Chaos/Public/GeometryCollection/Facades/CollectionVertexBoneWeightsFacade.h
+++ b/Engine/Source/Runtime/Experimental/Chaos/Public/GeometryCollection/Facades/CollectionVertexBoneWeightsFacade.h
@@ -55,13 +55,18 @@ namespace GeometryCollection::Facades
 		/** Add bone weight based on the kinematic bindings. */
 		CHAOS_API void AddBoneWeightsFromKinematicBindings();
 
-		/** Add bone weight based on the kinematic bindings. */
+		/** Add single bone/weight to vertex */
 		CHAOS_API void AddBoneWeight(int32 VertexIndex, int32 BoneIndex, float BoneWeight);
 
+		/** Modify bone weight based on the kinematic bindings. */
+		CHAOS_API void ModifyBoneWeight(int32 VertexIndex, TArray<int32> VertexBoneIndex, TArray<float> VertexBoneWeight);
+
 		/** Return the vertex bone indices from the collection. Null if not initialized.  */
 		const TManagedArray< TArray<int32> >* FindBoneIndices()  const { return BoneIndexAttribute.Find(); }
 		const TManagedArray< TArray<int32> >& GetBoneIndices() const { return BoneIndexAttribute.Get(); }
 
+		/** Return if the vertex is kinematic */
+		CHAOS_API bool IsKinematicVertex(int32 VertexIndex) const { return BoneIndexAttribute.IsValidIndex(VertexIndex) && BoneIndexAttribute[VertexIndex].Num(); };
 
 		/** Return the vertex bone weights from the collection. Null if not initialized. */
 		const TManagedArray< TArray<float> >* FindBoneWeights()  const { return BoneWeightAttribute.Find(); }
