ChaosFlesh: CIS fix
Embedded positions are now initialized to 0 even without embedding info.

diff --git a/Engine/Plugins/Experimental/ChaosFlesh/Source/ChaosFleshEngine/Private/ChaosFlesh/ChaosDeformableTetrahedralComponent.cpp b/Engine/Plugins/Experimental/ChaosFlesh/Source/ChaosFleshEngine/Private/ChaosFlesh/ChaosDeformableTetrahedralComponent.cpp
--- a/Engine/Plugins/Experimental/ChaosFlesh/Source/ChaosFleshEngine/Private/ChaosFlesh/ChaosDeformableTetrahedralComponent.cpp
+++ b/Engine/Plugins/Experimental/ChaosFlesh/Source/ChaosFleshEngine/Private/ChaosFlesh/ChaosDeformableTetrahedralComponent.cpp
@@ -978,7 +978,8 @@ TArray<FVector> UDeformableTetrahedralComponent::GetEmbeddedPositionsInternal(
 	auto UEVert3d = [](FVector3f V) { return FVector3d(V.X, V.Y, V.Z); };
 	auto UEVert4d = [](FVector4f V) { return FVector4d(V.X, V.Y, V.Z, V.W); };
 
-	TArray<FVector> OutPositions = InPositions;
+	TArray<FVector> OutPositions;
+	OutPositions.Init(FVector(0.f), InPositions.Num());
 
 	if (const UFleshAsset* RestAsset = GetRestCollection())
 	{
