ChaosFlesh: fixed bugs in CreateTetrahedron

diff --git a/Engine/Plugins/Experimental/ChaosCachingUSD/Source/ChaosCachingUSD/Private/ChaosCachingUSD/Operations.cpp b/Engine/Plugins/Experimental/ChaosCachingUSD/Source/ChaosCachingUSD/Private/ChaosCachingUSD/Operations.cpp
--- a/Engine/Plugins/Experimental/ChaosCachingUSD/Source/ChaosCachingUSD/Private/ChaosCachingUSD/Operations.cpp
+++ b/Engine/Plugins/Experimental/ChaosCachingUSD/Source/ChaosCachingUSD/Private/ChaosCachingUSD/Operations.cpp
@@ -481,7 +481,7 @@ UE::ChaosCachingUSD::WriteTetMesh(
 	// bool does the same thing.  What's more is that we currently have no consumer for the tet 
 	// topology, so this really doesn't matter.  At least not yet.
 	pxr::UEUsdGeomTetMesh TetMesh = pxr::UEUsdGeomTetMesh::Define(Stage, Path);
-	if (/*TetMesh*/true)
+	if (Tetrahedron->Num())
 	{
 		pxr::UsdPrim Prim = TetMesh.GetPrim();
 
diff --git a/Engine/Plugins/Experimental/ChaosFlesh/Source/ChaosFleshNodes/Private/Dataflow/ChaosFleshCreateTetrahedronNode.cpp b/Engine/Plugins/Experimental/ChaosFlesh/Source/ChaosFleshNodes/Private/Dataflow/ChaosFleshCreateTetrahedronNode.cpp
--- a/Engine/Plugins/Experimental/ChaosFlesh/Source/ChaosFleshNodes/Private/Dataflow/ChaosFleshCreateTetrahedronNode.cpp
+++ b/Engine/Plugins/Experimental/ChaosFlesh/Source/ChaosFleshNodes/Private/Dataflow/ChaosFleshCreateTetrahedronNode.cpp
@@ -283,7 +283,7 @@ void FCreateTetrahedronDataflowNode::EvaluateTetWild(
 		TArray<FIntVector4> Tets;
 		FProgressCancel Progress;
 		UE_LOG(LogChaosFlesh, Display,TEXT("Generating tet mesh via TetWild..."));
-		if (UE::Geometry::FTetWild::ComputeTetMesh(Params, Verts, Tris, TetVerts, Tets, &Progress))
+		if (UE::Geometry::FTetWild::ComputeTetMesh(Params, Verts, Tris, TetVerts, Tets, &Progress) && Tets.Num() > 0)
 		{
 			TArray<FIntVector3> SurfaceElements = UE::Dataflow::GetSurfaceTriangles(Tets, !bDiscardInteriorTriangles);
 			TUniquePtr<FTetrahedralCollection> TetCollection(FTetrahedralCollection::NewTetrahedralCollection(TetVerts, SurfaceElements, Tets));
