Chaos Flesh - Fixed the default SimCache path so that writing caches doesn't fail when Unreal Editor is running from a launcher installation.

[RN] Chaos Flesh - Fixed the default SimCache path so that written caches go in the project Saved folder instead of the Unreal Editor install folder.

diff --git a/Engine/Plugins/Experimental/ChaosFlesh/Source/ChaosFleshEngine/Private/ChaosCache/FleshComponentCacheAdapter.cpp b/Engine/Plugins/Experimental/ChaosFlesh/Source/ChaosFleshEngine/Private/ChaosCache/FleshComponentCacheAdapter.cpp
--- a/Engine/Plugins/Experimental/ChaosFlesh/Source/ChaosFleshEngine/Private/ChaosCache/FleshComponentCacheAdapter.cpp
+++ b/Engine/Plugins/Experimental/ChaosFlesh/Source/ChaosFleshEngine/Private/ChaosCache/FleshComponentCacheAdapter.cpp
@@ -555,7 +555,7 @@ namespace Chaos
 		FString CacheDir = InObserved.USDCacheDirectory.Path;
 		if (CacheDir.IsEmpty())
 		{
-			CacheDir = FString(TEXT("SimCache"));
+			CacheDir = FPaths::Combine(FPaths::ProjectSavedDir(), FString(TEXT("SimCache")));
 		}
 		FPaths::NormalizeDirectoryName(CacheDir);
 		if (CacheDir.StartsWith(FString(TEXT("/Game"))))
