ChaosFlesh: minor fixes in flesh dataflow nodes

diff --git a/Engine/Plugins/Experimental/ChaosFlesh/Source/ChaosFleshNodes/Public/Dataflow/ChaosFleshComputeMuscleActivationNode.h b/Engine/Plugins/Experimental/ChaosFlesh/Source/ChaosFleshNodes/Public/Dataflow/ChaosFleshComputeMuscleActivationNode.h
--- a/Engine/Plugins/Experimental/ChaosFlesh/Source/ChaosFleshNodes/Public/Dataflow/ChaosFleshComputeMuscleActivationNode.h
+++ b/Engine/Plugins/Experimental/ChaosFlesh/Source/ChaosFleshNodes/Public/Dataflow/ChaosFleshComputeMuscleActivationNode.h
@@ -17,19 +17,20 @@ struct FComputeMuscleActivationDataNode : public FDataflowNode
 {
 	GENERATED_USTRUCT_BODY()
 	DATAFLOW_NODE_DEFINE_INTERNAL(FComputeMuscleActivationDataNode, "ComputeMuscleActivationData", "Flesh", "")
+	DATAFLOW_NODE_RENDER_TYPE("SurfaceRender", FGeometryCollection::StaticType(), "Collection")
 
 public:
 	UPROPERTY(meta = (DataflowInput, DataflowOutput, DisplayName = "Collection", DataflowPassthrough = "Collection"))
 	FManagedArrayCollection Collection;
 
 	UPROPERTY(meta = (DataflowInput, DisplayName = "OriginVertexIndices"))
-		TArray<int32> OriginIndicesIn;
+	TArray<int32> OriginIndicesIn;
 
 	UPROPERTY(meta = (DataflowInput, DisplayName = "InsertionVertexIndices"))
-		TArray<int32> InsertionIndicesIn;
+	TArray<int32> InsertionIndicesIn;
 
 	UPROPERTY(EditAnywhere, Category = "Dataflow")
-		float ContractionVolumeScale = 1.f;
+	float ContractionVolumeScale = 1.f;
 
 	FComputeMuscleActivationDataNode(const UE::Dataflow::FNodeParameters& InParam, FGuid InGuid = FGuid::NewGuid())
 		: FDataflowNode(InParam, InGuid)
diff --git a/Engine/Plugins/Experimental/ChaosFlesh/Source/ChaosFleshNodes/Public/Dataflow/ChaosFleshGetFleshAssetNode.h b/Engine/Plugins/Experimental/ChaosFlesh/Source/ChaosFleshNodes/Public/Dataflow/ChaosFleshGetFleshAssetNode.h
--- a/Engine/Plugins/Experimental/ChaosFlesh/Source/ChaosFleshNodes/Public/Dataflow/ChaosFleshGetFleshAssetNode.h
+++ b/Engine/Plugins/Experimental/ChaosFlesh/Source/ChaosFleshNodes/Public/Dataflow/ChaosFleshGetFleshAssetNode.h
@@ -14,7 +14,7 @@ struct FGetFleshAssetDataflowNode : public FDataflowNode
 {
 	GENERATED_USTRUCT_BODY()
 	DATAFLOW_NODE_DEFINE_INTERNAL(FGetFleshAssetDataflowNode, "GetFleshAsset", "Flesh", "")
-	DATAFLOW_NODE_RENDER_TYPE("SurfaceRender", FGeometryCollection::StaticType(), "Collection")
+	DATAFLOW_NODE_RENDER_TYPE("SurfaceRender", FGeometryCollection::StaticType(), "Output")
 
 public:
 
diff --git a/Engine/Plugins/Experimental/ChaosFlesh/Source/ChaosFleshNodes/Public/Dataflow/ChaosFleshSetFleshBonePositionTargetBindingNode.h b/Engine/Plugins/Experimental/ChaosFlesh/Source/ChaosFleshNodes/Public/Dataflow/ChaosFleshSetFleshBonePositionTargetBindingNode.h
--- a/Engine/Plugins/Experimental/ChaosFlesh/Source/ChaosFleshNodes/Public/Dataflow/ChaosFleshSetFleshBonePositionTargetBindingNode.h
+++ b/Engine/Plugins/Experimental/ChaosFlesh/Source/ChaosFleshNodes/Public/Dataflow/ChaosFleshSetFleshBonePositionTargetBindingNode.h
@@ -36,10 +36,10 @@ public:
 	UPROPERTY(EditAnywhere, Category = "Dataflow")
 		ESkeletalBindingMode SkeletalBindingMode = ESkeletalBindingMode::Dataflow_SkeletalBinding_PositionTarget;
 
-	UPROPERTY(EditAnywhere, Category = "Dataflow")
+	UPROPERTY(EditAnywhere, Category = "Dataflow", meta = (ClampMin = "0.0", EditCondition = "SkeletalBindingMode == ESkeletalBindingMode::Dataflow_SkeletalBinding_PositionTarget", EditConditionHides))
 		float PositionTargetStiffness = 10000.f;
 
-	UPROPERTY(meta = (DataflowInput, DisplayName = "SkeletalMesh"))
+	UPROPERTY(EditAnywhere, Category = "Dataflow", meta = (DataflowInput, DisplayName = "SkeletalMesh"))
 		TObjectPtr<USkeletalMesh> SkeletalMeshIn = nullptr;
 
 	UPROPERTY(EditAnywhere, Category = "Dataflow", meta = (ClampMin = "0.0", ClampMax = "2.0", UIMin = "0.0", UIMax = "2.0"))
