ChaosFlesh: fixed unset animation after editor restart

diff --git a/Engine/Plugins/Experimental/ChaosFlesh/Source/ChaosFleshEngine/Private/ChaosFlesh/FleshAsset.cpp b/Engine/Plugins/Experimental/ChaosFlesh/Source/ChaosFleshEngine/Private/ChaosFlesh/FleshAsset.cpp
--- a/Engine/Plugins/Experimental/ChaosFlesh/Source/ChaosFleshEngine/Private/ChaosFlesh/FleshAsset.cpp
+++ b/Engine/Plugins/Experimental/ChaosFlesh/Source/ChaosFleshEngine/Private/ChaosFlesh/FleshAsset.cpp
@@ -96,7 +96,7 @@ void UFleshAsset::WriteDataflowContent(const TObjectPtr<UDataflowBaseContent>& D
 		SkeletalContent->SetSkeletalMesh(SkeletalMesh, true);
 
 #if WITH_EDITORONLY_DATA
-		SkeletalContent->SetAnimationAsset(PreviewAnimationAsset.Get());
+		SkeletalContent->SetAnimationAsset(PreviewAnimationAsset.LoadSynchronous());
 		SkeletalContent->SolverTiming = PreviewSolverTiming;
 		SkeletalContent->SolverEvolution = PreviewSolverEvolution;
 		SkeletalContent->SolverCollisions = PreviewSolverCollisions;
