ChaosFlesh: - allow empty embedding in surface bindings
- fixed bone name matching in KinematicSkeleton
- deleted unused UProperty in AddKinematicParticles

diff --git a/Engine/Plugins/Experimental/ChaosFlesh/Source/ChaosFleshEngine/Private/ChaosFlesh/ChaosDeformableTetrahedralComponent.cpp b/Engine/Plugins/Experimental/ChaosFlesh/Source/ChaosFleshEngine/Private/ChaosFlesh/ChaosDeformableTetrahedralComponent.cpp
--- a/Engine/Plugins/Experimental/ChaosFlesh/Source/ChaosFleshEngine/Private/ChaosFlesh/ChaosDeformableTetrahedralComponent.cpp
+++ b/Engine/Plugins/Experimental/ChaosFlesh/Source/ChaosFleshEngine/Private/ChaosFlesh/ChaosDeformableTetrahedralComponent.cpp
@@ -708,8 +708,7 @@ TOptional<TArray<int32>> UDeformableTetrahedralComponent::GetMeshImportVertexMap
 				const FVector3f& FleshPosition = Positions[VertexIndex];
 				if (!MLDPosition.Equals(FleshPosition, UE_KINDA_SMALL_NUMBER))
 				{
-					UE_LOG(LogDeformableTetrahedralComponentInternal, Warning, TEXT("SkeletalMeshAsset and FleshAsset have different vertex positions. Check if the assets have the same vertex order."));
-					return None;
+					UE_LOG(LogDeformableTetrahedralComponentInternal, Warning, TEXT("SkeletalMeshAsset and FleshAsset have different vertex positions. Check if the assets have the same vertex order or some points are not embedded."));
 				}
 			}
 		}
diff --git a/Engine/Plugins/Experimental/ChaosFlesh/Source/ChaosFleshNodes/Private/Dataflow/ChaosFleshKinematicSkeletonConstraintNode.cpp b/Engine/Plugins/Experimental/ChaosFlesh/Source/ChaosFleshNodes/Private/Dataflow/ChaosFleshKinematicSkeletonConstraintNode.cpp
--- a/Engine/Plugins/Experimental/ChaosFlesh/Source/ChaosFleshNodes/Private/Dataflow/ChaosFleshKinematicSkeletonConstraintNode.cpp
+++ b/Engine/Plugins/Experimental/ChaosFlesh/Source/ChaosFleshNodes/Private/Dataflow/ChaosFleshKinematicSkeletonConstraintNode.cpp
@@ -8,6 +8,7 @@
 #include "GeometryCollection/Facades/CollectionKinematicBindingFacade.h"
 #include "GeometryCollection/Facades/CollectionVertexBoneWeightsFacade.h"
 #include "GeometryCollection/GeometryCollectionAlgo.h"
+#include "GeometryCollection/Facades/CollectionTransformFacade.h"
 
 #include UE_INLINE_GENERATED_CPP_BY_NAME(ChaosFleshKinematicSkeletonConstraintNode)
 
@@ -48,6 +49,8 @@ void FKinematicSkeletonConstraintDataflowNode::Evaluate(UE::Dataflow::FContext&
 				}
 				return Skip;
 			};
+			GeometryCollection::Facades::FCollectionTransformFacade TransformFacade(*InCollection);
+			const TMap<FString, int32> BoneNameMap = TransformFacade.BoneNameIndexMap();
 
 			for (int32 b = 0; b < Skeleton->GetReferenceSkeleton().GetNum(); ++b)
 			{
@@ -102,9 +105,10 @@ void FKinematicSkeletonConstraintDataflowNode::Evaluate(UE::Dataflow::FContext&
 							//get local coords of bound verts
 							typedef GeometryCollection::Facades::FKinematicBindingFacade FKinematics;
 							FKinematics Kinematics(*InCollection.Get()); Kinematics.DefineSchema();
-							if (Kinematics.IsValid())
+							FString ParentBoneName = Skeleton->GetReferenceSkeleton().GetBoneName(ParentIndex).ToString();
+							if (Kinematics.IsValid() && BoneNameMap.Contains(ParentBoneName))
 							{
-								FKinematics::FBindingKey Binding = Kinematics.SetBoneBindings(ParentIndex, BoundVerts, BoundWeights);
+								FKinematics::FBindingKey Binding = Kinematics.SetBoneBindings(BoneNameMap[ParentBoneName], BoundVerts, BoundWeights);
 								TManagedArray<TArray<FVector3f>>& LocalPos = InCollection->AddAttribute<TArray<FVector3f>>("LocalPosition", Binding.GroupName);
 								Kinematics.AddKinematicBinding(Binding);
 								auto DoubleVert = [](FVector3f V) { return FVector3d(V.X, V.Y, V.Z); };
diff --git a/Engine/Plugins/Experimental/ChaosFlesh/Source/ChaosFleshNodes/Public/Dataflow/ChaosFleshAddKinematicParticlesNode.h b/Engine/Plugins/Experimental/ChaosFlesh/Source/ChaosFleshNodes/Public/Dataflow/ChaosFleshAddKinematicParticlesNode.h
--- a/Engine/Plugins/Experimental/ChaosFlesh/Source/ChaosFleshNodes/Public/Dataflow/ChaosFleshAddKinematicParticlesNode.h
+++ b/Engine/Plugins/Experimental/ChaosFlesh/Source/ChaosFleshNodes/Public/Dataflow/ChaosFleshAddKinematicParticlesNode.h
@@ -21,9 +21,6 @@ struct FAddKinematicParticlesDataflowNode : public FDataflowNode
 public:
 	typedef FManagedArrayCollection DataType;
 
-	UPROPERTY(EditAnywhere, Category = "Dataflow")
-		float Radius = 40.f;
-
 	UPROPERTY(EditAnywhere, Category = "Dataflow")
 		FTransform Transform;
 
