Chaosflesh: - support removing vertices and related data (e.g. tetrahedra, faces, weak constraints) and extracting a submesh. Added DeleteFleshVertices node.
- refactored SetFleshBonePositionTargetBinding node
- added atomic Gass-Seidel vertex failure tracker to reduce error message spamming
- fixed bugs in multiple files

diff --git a/Engine/Plugins/Experimental/ChaosFlesh/Source/ChaosFlesh/Private/ChaosFlesh/TetrahedralCollection.cpp b/Engine/Plugins/Experimental/ChaosFlesh/Source/ChaosFlesh/Private/ChaosFlesh/TetrahedralCollection.cpp
--- a/Engine/Plugins/Experimental/ChaosFlesh/Source/ChaosFlesh/Private/ChaosFlesh/TetrahedralCollection.cpp
+++ b/Engine/Plugins/Experimental/ChaosFlesh/Source/ChaosFlesh/Private/ChaosFlesh/TetrahedralCollection.cpp
@@ -1,6 +1,6 @@
 // Copyright Epic Games, Inc. All Rights Reserved.
 #include "ChaosFlesh/TetrahedralCollection.h"
-
+#include "GeometryCollection/Facades/CollectionPositionTargetFacade.h"
 
 DEFINE_LOG_CATEGORY_STATIC(FTetrahedralCollectionLogging, Log, All);
 
@@ -300,4 +300,90 @@ void FTetrahedralCollection::ReorderTetrahedralElements(const TArray<int32>& New
 	}
 
 	Super::ReorderElements(TetrahedralGroup, NewTetOrder);
+}
+
+void FTetrahedralCollection::RemoveVertices(const TArray<int32>& SortedVertexIndices)
+{
+	FProcessingParameters Param;
+	Param.bReindexDependentAttibutes = true;
+	Param.bDoValidation = false;
+	Super::RemoveElements(FGeometryCollection::VerticesGroup, SortedVertexIndices, Param);
+	const int32 NumVertices = NumElements(FGeometryCollection::VerticesGroup);
+	const int32 NumFaces = NumElements(FGeometryCollection::FacesGroup);
+	const int32 NumGeometry = NumElements(FGeometryCollection::GeometryGroup);
+	const int32 NumTets = NumElements(FTetrahedralCollection::TetrahedralGroup);
+	const int32 NumTransforms = NumElements(FTetrahedralCollection::TransformGroup);
+
+	// --- Tetrahedral Group --
+	TArray<int32> DeletedTetrahedralElements;
+	TArray<int32> TetrahedralElementsRemap;
+	TetrahedralElementsRemap.Init(INDEX_NONE, NumTets);
+	for (int32 Idx = 0; Idx < NumTets; ++Idx)
+	{
+		bool bKeepTet = true;
+		for (int32 LocalIdx = 0; LocalIdx < 4; ++LocalIdx)
+		{
+			if (Tetrahedron[Idx][LocalIdx] < 0)
+			{
+				bKeepTet = false;
+				DeletedTetrahedralElements.Add(Idx);
+				break;
+			}
+		}
+		if (bKeepTet)
+		{
+			TetrahedralElementsRemap[Idx] = Idx - DeletedTetrahedralElements.Num();
+		}
+	}
+	Super::RemoveElements(TetrahedralGroup, DeletedTetrahedralElements, Param);
+
+	// --- Faces Group --
+	TArray<int32> DeletedFaceElements;
+	for (int32 Idx = 0; Idx < NumFaces; ++Idx)
+	{
+		for (int32 LocalIdx = 0; LocalIdx < 3; ++LocalIdx)
+		{
+			if (Indices[Idx][LocalIdx] < 0)
+			{
+				DeletedFaceElements.Add(Idx);
+				break;
+			}
+		}
+	}
+	Super::RemoveElements(FGeometryCollection::FacesGroup, DeletedFaceElements, Param);
+
+	// --- Transform Group --
+	TArray<int32> DeletedTransformElements;
+	for (int32 Idx = 0; Idx < NumGeometry; ++Idx)
+	{
+		if (VertexCount[Idx] <= 0)
+		{
+			DeletedTransformElements.Add(TransformIndex[Idx]);
+		}
+	}
+	Param.bDoValidation = true; //Validate after the last RemoveElement
+	Super::RemoveElements(FGeometryCollection::TransformGroup, DeletedTransformElements, Param);
+
+	// --- Incident Elements (Tetrahedral Dependency) --
+	InitIncidentElements();
+	/*for (int32 Idx = 0; Idx < NumVertices; Idx++)
+	{
+		TArray<int32> NewIncidentElements;
+		TArray<int32> NewIncidentElementsLocal;
+		for (int32 i = 0; i < IncidentElements[Idx].Num(); ++i)
+		{
+			const int32 NewTetIndex = TetrahedralElementsRemap[IncidentElements[Idx][i]];
+			if (NewTetIndex >= 0)
+			{
+				NewIncidentElements.Add(NewTetIndex);
+				NewIncidentElementsLocal.Add(IncidentElementsLocalIndex[Idx][i]);
+			}
+		}
+		IncidentElements[Idx] = NewIncidentElements;
+		IncidentElementsLocalIndex[Idx] = NewIncidentElementsLocal;
+	}*/
+
+	// --- Position Targets (Vertex dependency)--
+	GeometryCollection::Facades::FPositionTargetFacade PositionTargetFacade(*this);
+	PositionTargetFacade.RemoveInvalidPositionTarget();
 }
\ No newline at end of file
diff --git a/Engine/Plugins/Experimental/ChaosFlesh/Source/ChaosFlesh/Public/ChaosFlesh/TetrahedralCollection.h b/Engine/Plugins/Experimental/ChaosFlesh/Source/ChaosFlesh/Public/ChaosFlesh/TetrahedralCollection.h
--- a/Engine/Plugins/Experimental/ChaosFlesh/Source/ChaosFlesh/Public/ChaosFlesh/TetrahedralCollection.h
+++ b/Engine/Plugins/Experimental/ChaosFlesh/Source/ChaosFlesh/Public/ChaosFlesh/TetrahedralCollection.h
@@ -57,6 +57,11 @@ public:
 
 	CHAOSFLESH_API void ReorderTetrahedralElements(const TArray<int32>& NewOrder);
 
+	/**
+	* Remove selected vertices.
+	*/
+	CHAOSFLESH_API void RemoveVertices(const TArray<int32>& SortedVertexIndices);
+
 	/*
 	*  Attribute Groups
 	*/
diff --git a/Engine/Plugins/Experimental/ChaosFlesh/Source/ChaosFleshNodes/ChaosFleshNodes.Build.cs b/Engine/Plugins/Experimental/ChaosFlesh/Source/ChaosFleshNodes/ChaosFleshNodes.Build.cs
--- a/Engine/Plugins/Experimental/ChaosFlesh/Source/ChaosFleshNodes/ChaosFleshNodes.Build.cs
+++ b/Engine/Plugins/Experimental/ChaosFlesh/Source/ChaosFleshNodes/ChaosFleshNodes.Build.cs
@@ -26,7 +26,8 @@ namespace UnrealBuildTool.Rules
 					"MeshConversionEngineTypes",
 					"MeshDescription",
 					"TetMeshing", 
-					"DataflowNodes"
+					"DataflowNodes",
+					"SkeletalMeshDescription",
 				}
 			);
         }
diff --git a/Engine/Plugins/Experimental/ChaosFlesh/Source/ChaosFleshNodes/Private/Dataflow/ChaosFleshAppendTetrahedralCollectionNode.cpp b/Engine/Plugins/Experimental/ChaosFlesh/Source/ChaosFleshNodes/Private/Dataflow/ChaosFleshAppendTetrahedralCollectionNode.cpp
--- a/Engine/Plugins/Experimental/ChaosFlesh/Source/ChaosFleshNodes/Private/Dataflow/ChaosFleshAppendTetrahedralCollectionNode.cpp
+++ b/Engine/Plugins/Experimental/ChaosFlesh/Source/ChaosFleshNodes/Private/Dataflow/ChaosFleshAppendTetrahedralCollectionNode.cpp
@@ -66,4 +66,19 @@ void FAppendTetrahedralCollectionDataflowNode::Evaluate(UE::Dataflow::FContext&
 		SetValue(Context, MoveTemp(GeometryGroupGuidsLocal1), &GeometryGroupGuidsOut1);
 		SetValue(Context, MoveTemp(GeometryGroupGuidsLocal2), &GeometryGroupGuidsOut2);
 	}
+}
+
+void FDeleteFleshVerticesDataflowNode::Evaluate(UE::Dataflow::FContext& Context, const FDataflowOutput* Out) const
+{
+	if (Out->IsA<FManagedArrayCollection>(&Collection))
+	{
+		TUniquePtr<FFleshCollection> InCollection(GetValue<FManagedArrayCollection>(Context, &Collection).NewCopy<FFleshCollection>());
+		if (IsConnected(&VertexSelection))
+		{
+			const FDataflowVertexSelection& InDataflowVertexSelection = GetValue<FDataflowVertexSelection>(Context, &VertexSelection);
+			TArray<int32> VertexSelected = InDataflowVertexSelection.AsArray();
+			InCollection->RemoveVertices(VertexSelected);
+		}
+		SetValue<const FManagedArrayCollection&>(Context, MoveTemp(*InCollection), &Collection);
+	}
 }
\ No newline at end of file
diff --git a/Engine/Plugins/Experimental/ChaosFlesh/Source/ChaosFleshNodes/Private/Dataflow/ChaosFleshNodesPlugin.cpp b/Engine/Plugins/Experimental/ChaosFlesh/Source/ChaosFleshNodes/Private/Dataflow/ChaosFleshNodesPlugin.cpp
--- a/Engine/Plugins/Experimental/ChaosFlesh/Source/ChaosFleshNodes/Private/Dataflow/ChaosFleshNodesPlugin.cpp
+++ b/Engine/Plugins/Experimental/ChaosFlesh/Source/ChaosFleshNodes/Private/Dataflow/ChaosFleshNodesPlugin.cpp
@@ -49,6 +49,7 @@ void IChaosFleshNodesPlugin::StartupModule()
 	DATAFLOW_NODE_REGISTER_CREATION_FACTORY(FComputeIslandsNode);
 	DATAFLOW_NODE_REGISTER_CREATION_FACTORY(FComputeMuscleActivationDataNode);
 	DATAFLOW_NODE_REGISTER_CREATION_FACTORY(FCreateTetrahedronDataflowNode);
+	DATAFLOW_NODE_REGISTER_CREATION_FACTORY(FDeleteFleshVerticesDataflowNode);
 	DATAFLOW_NODE_REGISTER_CREATION_FACTORY(FFleshAssetTerminalDataflowNode);
 	DATAFLOW_NODE_REGISTER_CREATION_FACTORY(FGenerateOriginInsertionNode);
 	DATAFLOW_NODE_REGISTER_CREATION_FACTORY(FGenerateSkeletalBindings);
diff --git a/Engine/Plugins/Experimental/ChaosFlesh/Source/ChaosFleshNodes/Private/Dataflow/ChaosFleshSetFleshBonePositionTargetBindingNode.cpp b/Engine/Plugins/Experimental/ChaosFlesh/Source/ChaosFleshNodes/Private/Dataflow/ChaosFleshSetFleshBonePositionTargetBindingNode.cpp
--- a/Engine/Plugins/Experimental/ChaosFlesh/Source/ChaosFleshNodes/Private/Dataflow/ChaosFleshSetFleshBonePositionTargetBindingNode.cpp
+++ b/Engine/Plugins/Experimental/ChaosFlesh/Source/ChaosFleshNodes/Private/Dataflow/ChaosFleshSetFleshBonePositionTargetBindingNode.cpp
@@ -11,6 +11,9 @@
 #include "GeometryCollection/Facades/CollectionTransformFacade.h"
 #include "Rendering/SkeletalMeshModel.h"
 #include "Rendering/SkeletalMeshRenderData.h"
+#include "SkeletalMeshAttributes.h"
+#include "ChaosFlesh/ChaosFlesh.h"
+
 #include "BoneWeights.h"
 
 #include UE_INLINE_GENERATED_CPP_BY_NAME(ChaosFleshSetFleshBonePositionTargetBindingNode)
@@ -22,7 +25,7 @@ void FSetFleshBonePositionTargetBindingDataflowNode::Evaluate(UE::Dataflow::FCon
 	if (Out->IsA<DataType>(&Collection))
 	{
 		DataType InCollection = GetValue<DataType>(Context, &Collection);
-
+#if WITH_EDITORONLY_DATA
 		if (TObjectPtr<USkeletalMesh> BoneSkeletalMesh = GetValue<TObjectPtr<USkeletalMesh>>(Context, &SkeletalMeshIn))
 		{
 			const TManagedArray<FIntVector>* Indices = InCollection.FindAttribute<FIntVector>("Indices", FGeometryCollection::FacesGroup);
@@ -30,46 +33,29 @@ void FSetFleshBonePositionTargetBindingDataflowNode::Evaluate(UE::Dataflow::FCon
 			const TManagedArray<FTransform3f>* Transform = InCollection.FindAttribute<FTransform3f>("Transform", FTransformCollection::TransformGroup);
 			const TManagedArray<FString>* TransformBoneName = InCollection.FindAttribute<FString>("BoneName", FTransformCollection::TransformGroup);
 			GeometryCollection::Facades::FCollectionTransformFacade TransformFacade(InCollection);
+			const FReferenceSkeleton& ReferenceSkeleton = BoneSkeletalMesh->GetRefSkeleton();
+			const int32 LODIndex = 0;
 			if (Indices && Vertices && Transform && TransformBoneName)
 			{
-				const TMap<FString, int32> BoneNameIndexMap = TransformFacade.BoneNameIndexMap();
 				FSkeletalMeshRenderData* RenderData = BoneSkeletalMesh->GetResourceForRendering();
-				if (RenderData->LODRenderData.Num())
+				FMeshDescription MeshDescription;
+				if (BoneSkeletalMesh->CloneMeshDescription(LODIndex, MeshDescription) && RenderData->LODRenderData.Num())
 				{
-					FSkeletalMeshLODRenderData* LODRenderData = &RenderData->LODRenderData[0];
-					const FPositionVertexBuffer& PositionVertexBuffer =
-						LODRenderData->StaticVertexBuffers.PositionVertexBuffer;
-
-					const FSkinWeightVertexBuffer* SkinWeightVertexBuffer = LODRenderData->GetSkinWeightVertexBuffer();
-					const int32 MaxBoneInfluences = SkinWeightVertexBuffer->GetMaxBoneInfluences();
-					TArray<FTransform> ComponentPose;
-					UE::Dataflow::Animation::GlobalTransforms(BoneSkeletalMesh->GetRefSkeleton(), ComponentPose);
-					TArray<TArray<int32>> BoneBoundVerts;
-					TArray<TArray<float>> BoneBoundWeights;
-					BoneBoundVerts.SetNum(ComponentPose.Num());
-					BoneBoundWeights.SetNum(ComponentPose.Num());
-
-					int32 NumSkeletonVertices = int32(PositionVertexBuffer.GetNumVertices());
+					FSkeletalMeshLODRenderData* LODRenderData = &RenderData->LODRenderData[LODIndex];
+					FSkeletalMeshAttributes MeshAttribs(MeshDescription);
+					FSkinWeightsVertexAttributesRef VertexSkinWeights = MeshAttribs.GetVertexSkinWeights();
+					int32 NumSkeletonVertices = MeshDescription.Vertices().Num();
 
 					Chaos::FReal SphereRadius = (Chaos::FReal)0.;
 
-					Chaos::TVec3<float> CoordMaxs(-FLT_MAX);
-					Chaos::TVec3<float> CoordMins(FLT_MAX);
+					Chaos::FVec3f CoordMaxs(-FLT_MAX);
+					Chaos::FVec3f CoordMins(FLT_MAX);
 					for (int32 i = 0; i < NumSkeletonVertices; i++)
 					{
-						for (int32 j = 0; j < 3; j++)
-						{
-							if (PositionVertexBuffer.VertexPosition(i)[j] > CoordMaxs[j])
-							{
-								CoordMaxs[j] = PositionVertexBuffer.VertexPosition(i)[j];
-							}
-							if (PositionVertexBuffer.VertexPosition(i)[j] < CoordMins[j])
-							{
-								CoordMins[j] = PositionVertexBuffer.VertexPosition(i)[j];
-							}
-						}
+						CoordMaxs = CoordMaxs.ComponentwiseMax(MeshDescription.GetVertexPosition(i));
+						CoordMins = CoordMins.ComponentwiseMin(MeshDescription.GetVertexPosition(i));
 					}
-					Chaos::TVec3<float> CoordDiff = (CoordMaxs - CoordMins) * VertexRadiusRatio;
+					Chaos::FVec3f CoordDiff = (CoordMaxs - CoordMins) * VertexRadiusRatio;
 					SphereRadius = Chaos::FReal(FGenericPlatformMath::Max(CoordDiff[0], FGenericPlatformMath::Max(CoordDiff[1], CoordDiff[2])));
 
 					TArray<Chaos::FSphere*> VertexSpherePtrs;
@@ -80,7 +66,7 @@ void FSetFleshBonePositionTargetBindingDataflowNode::Evaluate(UE::Dataflow::FCon
 
 					for (int32 i = 0; i < int32(NumSkeletonVertices); i++)
 					{
-						Chaos::TVec3<Chaos::FReal> SphereCenter(PositionVertexBuffer.VertexPosition(i));
+						Chaos::TVec3<Chaos::FReal> SphereCenter(MeshDescription.GetVertexPosition(i));
 						Chaos::FSphere VertexSphere(SphereCenter, SphereRadius);
 						VertexSpheres[i] = Chaos::FSphere(SphereCenter, SphereRadius);
 						VertexSpherePtrs[i] = &VertexSpheres[i];
@@ -91,63 +77,35 @@ void FSetFleshBonePositionTargetBindingDataflowNode::Evaluate(UE::Dataflow::FCon
 						Chaos::FReal,
 						3> VertexBVH(VertexSpherePtrs);
 
-					if (SkeletalBindingMode == ESkeletalBindingMode::Dataflow_SkeletalBinding_Kinematic)
-					{
+					GeometryCollection::Facades::FVertexBoneWeightsFacade VertexBoneWeightsFacade(InCollection);
 
-						GeometryCollection::Facades::FVertexBoneWeightsFacade VertexBoneWeightsFacade(InCollection);
+					
 
-						TArray<Chaos::TVector<int32, 3>> IndicesArray;
-						for (int32 i = 0; i < Indices->Num(); i++)
-						{
-							Chaos::TVector<int32, 3> CurrentIndices(0);
-							for (int32 j = 0; j < 3; j++)
-							{
-								CurrentIndices[j] = (*Indices)[i][j];
-							}
-							if (CurrentIndices[0] != -1
-								&& CurrentIndices[1] != -1
-								&& CurrentIndices[2] != -1)
-							{
-								IndicesArray.Emplace(CurrentIndices);
-							}
-						}
-						TArray<TArray<int32>> LocalIndex;
-						TArray<TArray<int32>>* LocalIndexPtr = &LocalIndex;
-						TArray<TArray<int>> GlobalIndex = Chaos::Utilities::ComputeIncidentElements(IndicesArray, LocalIndexPtr);
-						int32 ActualParticleCount = 0;
-						for (int32 l = 0; l < GlobalIndex.Num(); l++)
+					//Mapping from SKM bone index to Collection bone index
+					const TMap<FString, int32> BoneNameIndexMap = TransformFacade.BoneNameIndexMap();
+					TMap<int32, int32> SKMBoneIndexToCollectionBoneIndex;
+					for (int32 BoneIndex = 0; BoneIndex < ReferenceSkeleton.GetNum(); ++BoneIndex)
+					{
+						FString SKMBoneName = ReferenceSkeleton.GetRefBoneInfo()[BoneIndex].Name.ToString();
+						if (BoneNameIndexMap.Contains(SKMBoneName))
 						{
-							if (GlobalIndex[l].Num() > 0)
-							{
-								ActualParticleCount += 1;
-							}
-						}
-						TArray<Chaos::TVector<float, 3>> IndicesPositions;
-						IndicesPositions.SetNum(ActualParticleCount);
-						TArray<int32> IndicesMap;
-						IndicesMap.SetNum(ActualParticleCount);
-						int32 CurrentParticleIndex = 0;
-						for (int32 i = 0; i < GlobalIndex.Num(); i++)
-						{
-							if (GlobalIndex[i].Num() > 0)
-							{
-								IndicesPositions[CurrentParticleIndex] = (*Vertices)[(*Indices)[GlobalIndex[i][0]][LocalIndex[i][0]]];
-								IndicesMap[CurrentParticleIndex] = (*Indices)[GlobalIndex[i][0]][LocalIndex[i][0]];
-								CurrentParticleIndex += 1;
-							}
+							SKMBoneIndexToCollectionBoneIndex.Add(BoneIndex, BoneNameIndexMap[SKMBoneName]);
 						}
+					}
 
-						for (int32 i = 0; i < IndicesMap.Num(); i++)
+					if (SkeletalBindingMode == ESkeletalBindingMode::Dataflow_SkeletalBinding_Kinematic)
+					{
+						for (int32 i = 0; i < Vertices->Num(); ++i)
 						{
 							//only work on particles that are not kinematic:
-							if (!VertexBoneWeightsFacade.IsKinematicVertex(IndicesMap[i]))
+							if (!VertexBoneWeightsFacade.IsKinematicVertex(i))
 							{
-								TArray<int32> ParticleIntersection = VertexBVH.FindAllIntersections((*Vertices)[IndicesMap[i]]);
+								TArray<int32> ParticleIntersection = VertexBVH.FindAllIntersections((*Vertices)[i]);
 								int32 MinIndex = -1;
 								float MinDis = SphereRadius;
 								for (int32 j = 0; j < ParticleIntersection.Num(); j++)
 								{
-									Chaos::FRealSingle CurrentDistance = ((*Vertices)[IndicesMap[i]] - PositionVertexBuffer.VertexPosition(ParticleIntersection[j])).Size();
+									Chaos::FRealSingle CurrentDistance = ((*Vertices)[i] - MeshDescription.GetVertexPosition(ParticleIntersection[j])).Size();
 									if (CurrentDistance < MinDis)
 									{
 										MinDis = CurrentDistance;
@@ -157,27 +115,20 @@ void FSetFleshBonePositionTargetBindingDataflowNode::Evaluate(UE::Dataflow::FCon
 
 								if (MinIndex != -1)
 								{
-									int32 BoneParticleIndex = -1;
-
-									int32 SectionIndex;
-									int32 VertIndex;
-									LODRenderData->GetSectionFromVertexIndex(MinIndex, SectionIndex, VertIndex);
-
-									check(SectionIndex < LODRenderData->RenderSections.Num());
-									const FSkelMeshRenderSection& Section = LODRenderData->RenderSections[SectionIndex];
-									int32 BufferVertIndex = Section.GetVertexBufferIndex() + VertIndex;
-									for (int32 InfluenceIndex = 0; InfluenceIndex < MaxBoneInfluences; InfluenceIndex++)
+									FVertexBoneWeights BoneWeights = VertexSkinWeights.Get(FVertexID(MinIndex));
+									const int32 InfluenceCount = BoneWeights.Num();
+									TArray<int32> VertexBoneIndex;
+									TArray<float> VertexBoneWeight;
+									for (int32 InfluenceIndex = 0; InfluenceIndex < InfluenceCount; ++InfluenceIndex)
 									{
-										const int32 BoneIndex = Section.BoneMap[SkinWeightVertexBuffer->GetBoneIndex(BufferVertIndex, InfluenceIndex)];
-										const float	Weight = (float)SkinWeightVertexBuffer->GetBoneWeight(BufferVertIndex, InfluenceIndex) * UE::AnimationCore::InvMaxRawBoneWeightFloat;
-										if (Weight > float(0) && 0 <= BoneIndex && BoneIndex < ComponentPose.Num())
+										if (SKMBoneIndexToCollectionBoneIndex.Contains(BoneWeights[InfluenceIndex].GetBoneIndex()))
 										{
-											BoneBoundVerts[BoneIndex].Add(IndicesMap[i]);
-											//Only rigid skinning for now
-											BoneBoundWeights[BoneIndex].Add(1.f);
-											break;
+											VertexBoneIndex.Add(SKMBoneIndexToCollectionBoneIndex[BoneWeights[InfluenceIndex].GetBoneIndex()]);
+											VertexBoneWeight.Add(BoneWeights[InfluenceIndex].GetWeight());
 										}
 									}
+									VertexBoneWeightsFacade.ModifyBoneWeight(i, VertexBoneIndex, VertexBoneWeight);
+									VertexBoneWeightsFacade.SetVertexKinematic(i);
 								}
 							}
 						}
@@ -188,141 +139,72 @@ void FSetFleshBonePositionTargetBindingDataflowNode::Evaluate(UE::Dataflow::FCon
 						GeometryCollection::Facades::FPositionTargetFacade PositionTargets(InCollection);
 						PositionTargets.DefineSchema();
 
-						for (int32 i = 0; i < Indices->Num(); i++)
+						for (int32 i = 0; i < Vertices->Num(); ++i)
 						{
-							TArray<int32> TriangleIntersections0 = VertexBVH.FindAllIntersections((*Vertices)[(*Indices)[i][0]]);
-							TArray<int32> TriangleIntersections1 = VertexBVH.FindAllIntersections((*Vertices)[(*Indices)[i][1]]);
-							TArray<int32> TriangleIntersections2 = VertexBVH.FindAllIntersections((*Vertices)[(*Indices)[i][2]]);
-							TriangleIntersections0.Sort();
-							TriangleIntersections1.Sort();
-							TriangleIntersections2.Sort();
-							TArray<int32> TriangleIntersections({});
-							for (int32 k = 0; k < TriangleIntersections0.Num(); k++)
-							{
-								if (TriangleIntersections1.Contains(TriangleIntersections0[k])
-									&& TriangleIntersections2.Contains(TriangleIntersections0[k]))
-								{
-									TriangleIntersections.Emplace(TriangleIntersections0[k]);
-								}
-							}
-							int32 MinIndex = -1;
-							float MinDis = SphereRadius;
-							Chaos::TVector<float, 3> ClosestBary(0.f);
-							for (int32 j = 0; j < TriangleIntersections.Num(); j++)
+							//only work on particles that are not kinematic:
+							if (!VertexBoneWeightsFacade.IsKinematicVertex(i))
 							{
-								Chaos::TVector<float, 3> Bary,
-									TriPos0((*Vertices)[(*Indices)[i][0]]), TriPos1((*Vertices)[(*Indices)[i][1]]), TriPos2((*Vertices)[(*Indices)[i][2]]),
-									ParticlePos(PositionVertexBuffer.VertexPosition(TriangleIntersections[j]));
-								Chaos::TVector<Chaos::FRealSingle, 3> ClosestPoint = Chaos::FindClosestPointAndBaryOnTriangle(TriPos0, TriPos1, TriPos2, ParticlePos, Bary);
-								Chaos::FRealSingle CurrentDistance = (ParticlePos - ClosestPoint).Size();
-								if (CurrentDistance < MinDis)
+								TArray<int32> ParticleIntersection = VertexBVH.FindAllIntersections((*Vertices)[i]);
+								int32 MinIndex = -1;
+								float MinDis = SphereRadius;
+								for (int32 j = 0; j < ParticleIntersection.Num(); j++)
 								{
-									MinDis = CurrentDistance;
-									MinIndex = TriangleIntersections[j];
-									ClosestBary = Bary;
+									Chaos::FRealSingle CurrentDistance = ((*Vertices)[i] - MeshDescription.GetVertexPosition(ParticleIntersection[j])).Size();
+									if (CurrentDistance < MinDis)
+									{
+										MinDis = CurrentDistance;
+										MinIndex = ParticleIntersection[j];
+									}
 								}
-							}
 
-							if (MinIndex != -1)
-							{
-								int32 BoneParticleIndex = -1;
-								//TODO: add kinematic particles first
-								if (!BoneVertexToCollectionMap.Contains(MinIndex))
+								if (MinIndex != -1)
 								{
-
+									// add kinematic particles
 									int32 ParticleIndex = InCollection.AddElements(1, FGeometryCollection::VerticesGroup);
 									TManagedArray<FVector3f>& CurrentVertices = InCollection.ModifyAttribute<FVector3f>("Vertex", FGeometryCollection::VerticesGroup);
 
-									CurrentVertices[ParticleIndex] = PositionVertexBuffer.VertexPosition(MinIndex);
-									BoneVertexToCollectionMap.Emplace(MinIndex, ParticleIndex);
-
-									int32 SectionIndex;
-									int32 VertIndex;
-									LODRenderData->GetSectionFromVertexIndex(MinIndex, SectionIndex, VertIndex);
+									CurrentVertices[ParticleIndex] = MeshDescription.GetVertexPosition(MinIndex);
 
-									check(SectionIndex < LODRenderData->RenderSections.Num());
-									const FSkelMeshRenderSection& Section = LODRenderData->RenderSections[SectionIndex];
-									int32 BufferVertIndex = Section.GetVertexBufferIndex() + VertIndex;
-									for (int32 InfluenceIndex = 0; InfluenceIndex < MaxBoneInfluences; InfluenceIndex++)
+									FVertexBoneWeights BoneWeights = VertexSkinWeights.Get(FVertexID(MinIndex));
+									const int32 InfluenceCount = BoneWeights.Num();
+									TArray<int32> VertexBoneIndex;
+									TArray<float> VertexBoneWeight;
+									for (int32 InfluenceIndex = 0; InfluenceIndex < InfluenceCount; ++InfluenceIndex)
 									{
-										const int32 BoneIndex = Section.BoneMap[SkinWeightVertexBuffer->GetBoneIndex(BufferVertIndex, InfluenceIndex)];
-										const float	Weight = (float)SkinWeightVertexBuffer->GetBoneWeight(BufferVertIndex, InfluenceIndex) * UE::AnimationCore::InvMaxRawBoneWeightFloat;
-										if (Weight > float(0) && 0 <= BoneIndex && BoneIndex < ComponentPose.Num())
+										if (SKMBoneIndexToCollectionBoneIndex.Contains(BoneWeights[InfluenceIndex].GetBoneIndex()))
 										{
-											BoneBoundVerts[BoneIndex].Add(ParticleIndex);
-											BoneBoundWeights[BoneIndex].Add(Weight);
+											VertexBoneIndex.Add(SKMBoneIndexToCollectionBoneIndex[BoneWeights[InfluenceIndex].GetBoneIndex()]);
+											VertexBoneWeight.Add(BoneWeights[InfluenceIndex].GetWeight());
 										}
 									}
-								}
-
-								BoneParticleIndex = BoneVertexToCollectionMap[MinIndex];
-
-								GeometryCollection::Facades::FPositionTargetsData DataPackage;
-								DataPackage.TargetIndex.Init(BoneParticleIndex, 1);
-								DataPackage.TargetWeights.Init(1.f, 1);
-								DataPackage.SourceWeights.Init(1.f, 3);
-								DataPackage.SourceIndex.Init(-1, 3);
-								DataPackage.SourceIndex[0] = (*Indices)[i][0];
-								DataPackage.SourceIndex[1] = (*Indices)[i][1];
-								DataPackage.SourceIndex[2] = (*Indices)[i][2];
-								DataPackage.SourceWeights[0] = ClosestBary[0];
-								DataPackage.SourceWeights[1] = ClosestBary[1];
-								DataPackage.SourceWeights[2] = ClosestBary[2];
-								if (const TManagedArray<float>* Mass = InCollection.FindAttribute<float>("Mass", FGeometryCollection::VerticesGroup))
-								{
-									DataPackage.Stiffness = 0.f;
-									//Target is kinematic, only compute stiffness from source
-									for (int32 k = 0; k < 3; k++)
+									VertexBoneWeightsFacade.ModifyBoneWeight(ParticleIndex, VertexBoneIndex, VertexBoneWeight);
+									VertexBoneWeightsFacade.SetVertexKinematic(ParticleIndex);
+
+									GeometryCollection::Facades::FPositionTargetsData DataPackage;
+									DataPackage.TargetIndex.Init(ParticleIndex, 1);
+									DataPackage.TargetWeights.Init(1.f, 1);
+									DataPackage.SourceWeights.Init(1.f, 1);
+									DataPackage.SourceIndex.Init(i, 1);
+									if (const TManagedArray<float>* Mass = InCollection.FindAttribute<float>("Mass", FGeometryCollection::VerticesGroup))
 									{
-										DataPackage.Stiffness += DataPackage.SourceWeights[k] * PositionTargetStiffness * (*Mass)[DataPackage.SourceIndex[k]];
+										//Target is kinematic, only compute stiffness from source
+										DataPackage.Stiffness = PositionTargetStiffness * (*Mass)[i];
 									}
-								}
-								else
-								{
-									DataPackage.Stiffness = PositionTargetStiffness;
-								}
-								PositionTargets.AddPositionTarget(DataPackage);
-							}
-
-
-						}
-
-
-					}
-					auto DoubleVert = [](FVector3f V) { return FVector3d(V.X, V.Y, V.Z); };
-
-					for (int32 BoneIndex = 0; BoneIndex < ComponentPose.Num(); ++BoneIndex)
-					{
-						if (BoneBoundVerts[BoneIndex].Num())
-						{
-							FString BoneName = BoneSkeletalMesh->GetRefSkeleton().GetBoneName(BoneIndex).ToString();
-							//get local coords of bound verts
-							using FKinematics = GeometryCollection::Facades::FKinematicBindingFacade;
-							FKinematics Kinematics(InCollection); Kinematics.DefineSchema();
-							if (Kinematics.IsValid())
-							{
-								if (ensure(BoneNameIndexMap.Contains(BoneName)))
-								{
-									FKinematics::FBindingKey Binding = Kinematics.SetBoneBindings(BoneNameIndexMap[BoneName], BoneBoundVerts[BoneIndex], BoneBoundWeights[BoneIndex]);
-									TManagedArray<TArray<FVector3f>>& LocalPos = InCollection.AddAttribute<TArray<FVector3f>>("LocalPosition", Binding.GroupName);
-									Kinematics.AddKinematicBinding(Binding);
-
-									auto FloatVert = [](FVector3d V) { return FVector3f(V.X, V.Y, V.Z); };
-									LocalPos[Binding.Index].SetNum(BoneBoundVerts[BoneIndex].Num());
-									for (int32 i = 0; i < BoneBoundVerts[BoneIndex].Num(); i++)
+									else
 									{
-										FVector3f Temp = (*Vertices)[BoneBoundVerts[BoneIndex][i]];
-										LocalPos[Binding.Index][i] = FloatVert(ComponentPose[BoneIndex].InverseTransformPosition(DoubleVert(Temp)));
+										DataPackage.Stiffness = PositionTargetStiffness;
 									}
+									PositionTargets.AddPositionTarget(DataPackage);
 								}
 							}
 						}
 					}
-					GeometryCollection::Facades::FVertexBoneWeightsFacade(InCollection).AddBoneWeightsFromKinematicBindings();
-
 				}
 			}
 		}
+#else
+		ensureMsgf(false, TEXT("SetFleshBonePositionTargetBinding is an editor only node."));
+#endif
 		SetValue(Context, MoveTemp(InCollection), &Collection);
 	}
 }
\ No newline at end of file
diff --git a/Engine/Plugins/Experimental/ChaosFlesh/Source/ChaosFleshNodes/Public/Dataflow/ChaosFleshAppendTetrahedralCollectionNode.h b/Engine/Plugins/Experimental/ChaosFlesh/Source/ChaosFleshNodes/Public/Dataflow/ChaosFleshAppendTetrahedralCollectionNode.h
--- a/Engine/Plugins/Experimental/ChaosFlesh/Source/ChaosFleshNodes/Public/Dataflow/ChaosFleshAppendTetrahedralCollectionNode.h
+++ b/Engine/Plugins/Experimental/ChaosFlesh/Source/ChaosFleshNodes/Public/Dataflow/ChaosFleshAppendTetrahedralCollectionNode.h
@@ -57,5 +57,32 @@ public:
 		RegisterOutputConnection(&GeometrySelection2);
 	}
 
+	virtual void Evaluate(UE::Dataflow::FContext& Context, const FDataflowOutput* Out) const override;
+};
+
+/**
+ * Extract a Tetrahedral Collection from this collection based on selected vertex. Compatible attributes will be copied.
+ */
+USTRUCT(meta = (DataflowFlesh))
+struct FDeleteFleshVerticesDataflowNode : public FDataflowNode
+{
+	GENERATED_USTRUCT_BODY()
+	DATAFLOW_NODE_DEFINE_INTERNAL(FDeleteFleshVerticesDataflowNode, "DeleteFleshVertices", "Flesh", "")
+	DATAFLOW_NODE_RENDER_TYPE("SurfaceRender", FGeometryCollection::StaticType(), "Collection")
+public:
+	UPROPERTY(meta = (DataflowInput, DataflowOutput, DataflowPassthrough = "Collection"));
+	FManagedArrayCollection Collection;
+
+	UPROPERTY(meta = (DataflowInput))
+	FDataflowVertexSelection VertexSelection;
+
+	FDeleteFleshVerticesDataflowNode(const UE::Dataflow::FNodeParameters& InParam, FGuid InGuid = FGuid::NewGuid())
+		: FDataflowNode(InParam, InGuid)
+	{
+		RegisterInputConnection(&Collection);
+		RegisterInputConnection(&VertexSelection);
+		RegisterOutputConnection(&Collection, &Collection);
+	}
+
 	virtual void Evaluate(UE::Dataflow::FContext& Context, const FDataflowOutput* Out) const override;
 };
\ No newline at end of file
diff --git a/Engine/Plugins/Experimental/ChaosFlesh/Source/ChaosFleshNodes/Public/Dataflow/ChaosFleshKinematicMuscleAttachmentsNode.h b/Engine/Plugins/Experimental/ChaosFlesh/Source/ChaosFleshNodes/Public/Dataflow/ChaosFleshKinematicMuscleAttachmentsNode.h
--- a/Engine/Plugins/Experimental/ChaosFlesh/Source/ChaosFleshNodes/Public/Dataflow/ChaosFleshKinematicMuscleAttachmentsNode.h
+++ b/Engine/Plugins/Experimental/ChaosFlesh/Source/ChaosFleshNodes/Public/Dataflow/ChaosFleshKinematicMuscleAttachmentsNode.h
@@ -28,7 +28,7 @@ public:
 	UPROPERTY(meta = (DataflowInput, DisplayName = "InsertionVertexIndices"))
 		TArray<int32> InsertionVertexIndicesIn;
 
-	UPROPERTY(meta = (DataflowInput, DisplayName = "SkeletalMesh"))
+	UPROPERTY(EditAnywhere, Category = "Dataflow", meta = (DataflowInput, DisplayName = "SkeletalMesh"))
 		TObjectPtr<USkeletalMesh> SkeletalMeshIn;
 
 
diff --git a/Engine/Plugins/Experimental/GeometryCollectionPlugin/Source/GeometryCollectionNodes/Private/Dataflow/GeometryCollectionSelectionNodes.cpp b/Engine/Plugins/Experimental/GeometryCollectionPlugin/Source/GeometryCollectionNodes/Private/Dataflow/GeometryCollectionSelectionNodes.cpp
--- a/Engine/Plugins/Experimental/GeometryCollectionPlugin/Source/GeometryCollectionNodes/Private/Dataflow/GeometryCollectionSelectionNodes.cpp
+++ b/Engine/Plugins/Experimental/GeometryCollectionPlugin/Source/GeometryCollectionNodes/Private/Dataflow/GeometryCollectionSelectionNodes.cpp
@@ -1469,9 +1469,8 @@ void FGeometrySelectionToVertexSelectionDataflowNode::Evaluate(UE::Dataflow::FCo
 		}
 		else
 		{
-			const FString InGeometryIndices = GetValue<FString>(Context, &GeometryIndices);
 			TArray<FString> Indices;
-			InGeometryIndices.ParseIntoArray(Indices, TEXT(" "), true);
+			GeometryIndices.ParseIntoArray(Indices, TEXT(" "), true);
 			for (FString IndexStr : Indices)
 			{
 				if (IndexStr.IsNumeric())
diff --git a/Engine/Plugins/Experimental/GeometryCollectionPlugin/Source/GeometryCollectionNodes/Private/Dataflow/GeometryCollectionTransferVertexAttributeNode.cpp b/Engine/Plugins/Experimental/GeometryCollectionPlugin/Source/GeometryCollectionNodes/Private/Dataflow/GeometryCollectionTransferVertexAttributeNode.cpp
--- a/Engine/Plugins/Experimental/GeometryCollectionPlugin/Source/GeometryCollectionNodes/Private/Dataflow/GeometryCollectionTransferVertexAttributeNode.cpp
+++ b/Engine/Plugins/Experimental/GeometryCollectionPlugin/Source/GeometryCollectionNodes/Private/Dataflow/GeometryCollectionTransferVertexAttributeNode.cpp
@@ -464,7 +464,7 @@ namespace UE::Private
 			TriangleMesh.Init(SourceElements);
 			Chaos::FTriangleMesh::TSpatialHashType<Chaos::FReal> SpatialHash;
 			TriangleMesh.BuildSpatialHash(ConstComponentSpaceVertices, SpatialHash, SphereRadius);
-			for (int32 TargetIndex = 0; TargetIndex < ConstComponentSpaceVertices.Num(); ++TargetIndex)
+			for (int32 TargetIndex = 0; TargetIndex < ComponentSpaceTargetVertices.Num(); ++TargetIndex)
 			{
 				TArray<Chaos::TTriangleCollisionPoint<Chaos::FReal>> Result;
 				if (TriangleMesh.PointClosestTriangleQuery(SpatialHash, ConstComponentSpaceVertices,
diff --git a/Engine/Plugins/Experimental/GeometryCollectionPlugin/Source/GeometryCollectionNodes/Public/Dataflow/GeometryCollectionEditNodes.h b/Engine/Plugins/Experimental/GeometryCollectionPlugin/Source/GeometryCollectionNodes/Public/Dataflow/GeometryCollectionEditNodes.h
--- a/Engine/Plugins/Experimental/GeometryCollectionPlugin/Source/GeometryCollectionNodes/Public/Dataflow/GeometryCollectionEditNodes.h
+++ b/Engine/Plugins/Experimental/GeometryCollectionPlugin/Source/GeometryCollectionNodes/Public/Dataflow/GeometryCollectionEditNodes.h
@@ -81,11 +81,11 @@ public:
 	FManagedArrayCollection Collection;
 
 	/** Transform selection for setting visibility */
-	UPROPERTY(meta = (DataflowInput, DisplayName = "TransformSelection", DataflowIntrinsic))
+	UPROPERTY(meta = (DataflowInput, DataflowIntrinsic))
 	FDataflowTransformSelection TransformSelection;
 
 	/** Face selection for setting visibility */
-	UPROPERTY(meta = (DataflowInput, DisplayName = "TransformSelection", DataflowIntrinsic))
+	UPROPERTY(meta = (DataflowInput, DataflowIntrinsic))
 	FDataflowFaceSelection FaceSelection;
 
 	FSetVisibilityInCollectionDataflowNode(const UE::Dataflow::FNodeParameters& InParam, FGuid InGuid = FGuid::NewGuid())
diff --git a/Engine/Source/Runtime/Experimental/Chaos/Private/Chaos/Deformable/ChaosDeformableSolver.cpp b/Engine/Source/Runtime/Experimental/Chaos/Private/Chaos/Deformable/ChaosDeformableSolver.cpp
--- a/Engine/Source/Runtime/Experimental/Chaos/Private/Chaos/Deformable/ChaosDeformableSolver.cpp
+++ b/Engine/Source/Runtime/Experimental/Chaos/Private/Chaos/Deformable/ChaosDeformableSolver.cpp
@@ -589,7 +589,7 @@ namespace Chaos::Softs
 
 					for (int32 vdx : BoundVerts)
 					{
-						if (!VertexBoneWeightsFacade.IsKinematicVertex(vdx))
+						if (vdx >= 0 && !VertexBoneWeightsFacade.IsKinematicVertex(vdx))
 						{
 							if (!bHavePrintedLog)
 							{
diff --git a/Engine/Source/Runtime/Experimental/Chaos/Private/GeometryCollection/Facades/CollectionPositionTargetFacade.cpp b/Engine/Source/Runtime/Experimental/Chaos/Private/GeometryCollection/Facades/CollectionPositionTargetFacade.cpp
--- a/Engine/Source/Runtime/Experimental/Chaos/Private/GeometryCollection/Facades/CollectionPositionTargetFacade.cpp
+++ b/Engine/Source/Runtime/Experimental/Chaos/Private/GeometryCollection/Facades/CollectionPositionTargetFacade.cpp
@@ -124,6 +124,38 @@ namespace GeometryCollection::Facades
 	
 	}
 
-
-
+	int32 FPositionTargetFacade::RemoveInvalidPositionTarget()
+	{
+		check(!IsConst());
+		TArray<int32> InvalidConstraintIndices;
+		if (IsValid())
+		{
+			for (int32 Index = 0; Index < Collection->NumElements(GroupName); ++Index)
+			{
+				bool bDelete = false;
+				for (int32 TargetIdx : TargetIndexAttribute.Get()[Index])
+				{
+					if (TargetIdx < 0)
+					{
+						bDelete = true;
+						InvalidConstraintIndices.Add(Index);
+						break;
+					}
+				}
+				if (!bDelete)
+				{
+					for (int32 SourceIdx : SourceIndexAttribute.Get()[Index])
+					{
+						if (SourceIdx < 0)
+						{
+							InvalidConstraintIndices.Add(Index);
+							break;
+						}
+					}
+				}
+			}
+			Collection->RemoveElements(GroupName, InvalidConstraintIndices);
+		}
+		return InvalidConstraintIndices.Num();
+	}
 }
diff --git a/Engine/Source/Runtime/Experimental/Chaos/Private/GeometryCollection/TransformCollection.cpp b/Engine/Source/Runtime/Experimental/Chaos/Private/GeometryCollection/TransformCollection.cpp
--- a/Engine/Source/Runtime/Experimental/Chaos/Private/GeometryCollection/TransformCollection.cpp
+++ b/Engine/Source/Runtime/Experimental/Chaos/Private/GeometryCollection/TransformCollection.cpp
@@ -285,7 +285,7 @@ void FTransformCollection::MergeElements(const FName& Group, const TArray<int32>
 				const int32 Index = SortedMergeList[Sdx];
 				if (ensure(0 <= Index && Index < ParentArray.Num()))
 				{
-					const int32 RemapId = MergeRemapIndex[Index];
+					const int32 RemapId = MergeRemapIndex[Sdx];
 					const int32 ParentID = ParentArray[Index];
 					if (ensure(0 <= RemapId && RemapId < ParentArray.Num()))
 					{
diff --git a/Engine/Source/Runtime/Experimental/Chaos/Public/Chaos/Deformable/GaussSeidelMainConstraint.h b/Engine/Source/Runtime/Experimental/Chaos/Public/Chaos/Deformable/GaussSeidelMainConstraint.h
--- a/Engine/Source/Runtime/Experimental/Chaos/Public/Chaos/Deformable/GaussSeidelMainConstraint.h
+++ b/Engine/Source/Runtime/Experimental/Chaos/Public/Chaos/Deformable/GaussSeidelMainConstraint.h
@@ -152,7 +152,8 @@ namespace Chaos::Softs
 				ComputeNewtonResiduals(Particles, Dt, Write2File);
 			}
 			
-			
+			std::atomic<int32> ParticleFailureCounter(0);
+
 			for (int32 i = 0; i < ParticlesPerColor.Num(); i++)
 			{
 				int32 NumBatch = ParticlesPerColor[i].Num() / CorotatedParams.XPBDCorotatedBatchSize;
@@ -170,7 +171,10 @@ namespace Chaos::Softs
 								if (Particles.InvM(ParticlesPerColor[i][TaskIndex]) != T(0))
 								{
 									int32 ParticleIndex = ParticlesPerColor[i][TaskIndex];
-									ApplySingleParticle(ParticleIndex, Dt, Particles);
+									if (!ApplySingleParticle(ParticleIndex, Dt, Particles))
+									{
+										ParticleFailureCounter.fetch_add(1, std::memory_order_relaxed);
+									}
 								}
 							}
 						}
@@ -196,6 +200,11 @@ namespace Chaos::Softs
 			}
 			
 			CurrentIt ++;
+
+			if (ParticleFailureCounter.load(std::memory_order_relaxed) > 0)
+			{
+				UE_LOG(LogDeformableGaussSeidelMainConstraint, Warning, TEXT("%d Particle(s) are skipped because of too large dx size"), ParticleFailureCounter.load(std::memory_order_relaxed));
+			}
 		}
 
 		void InitStaticColor(const ParticleType& Particles, const TPBDActiveView<FSolverParticles>* InParticleActiveView = nullptr)
@@ -283,7 +292,7 @@ namespace Chaos::Softs
 
 	private:
 
-		void ApplySingleParticle(const int32 p, const T Dt, ParticleType& Particles)
+		bool ApplySingleParticle(const int32 p, const T Dt, ParticleType& Particles)
 		{
 			int32 ConstraintIndex = 0;
 			Chaos::TVector<T, 3> ParticleResidual((T)0.);
@@ -344,9 +353,10 @@ namespace Chaos::Softs
 				}
 				else
 				{
-					UE_LOG(LogDeformableGaussSeidelMainConstraint, Warning, TEXT("Following Particle is skipped because of too large dx size: %d"), p);
+					return false;
 				}
 			}
+			return true;
 		}
 		
 		void InitializeLambdas()
diff --git a/Engine/Source/Runtime/Experimental/Chaos/Public/GeometryCollection/Facades/CollectionPositionTargetFacade.h b/Engine/Source/Runtime/Experimental/Chaos/Public/GeometryCollection/Facades/CollectionPositionTargetFacade.h
--- a/Engine/Source/Runtime/Experimental/Chaos/Public/GeometryCollection/Facades/CollectionPositionTargetFacade.h
+++ b/Engine/Source/Runtime/Experimental/Chaos/Public/GeometryCollection/Facades/CollectionPositionTargetFacade.h
@@ -76,6 +76,9 @@ namespace GeometryCollection::Facades
 		CHAOS_API int32 AddPositionTarget(const FPositionTargetsData& InputData);
 		CHAOS_API FPositionTargetsData GetPositionTarget(const int32 DataIndex) const;
 		int32 NumPositionTargets() const { return TargetIndexAttribute.Num(); }
+		
+		/** Remove position targets with invalid indices*/
+		CHAOS_API int32 RemoveInvalidPositionTarget();
 
 	protected:
 		const FManagedArrayCollection& ConstCollection;
diff --git a/Engine/Source/Runtime/Experimental/Chaos/Public/GeometryCollection/GeometryCollection.h b/Engine/Source/Runtime/Experimental/Chaos/Public/GeometryCollection/GeometryCollection.h
--- a/Engine/Source/Runtime/Experimental/Chaos/Public/GeometryCollection/GeometryCollection.h
+++ b/Engine/Source/Runtime/Experimental/Chaos/Public/GeometryCollection/GeometryCollection.h
@@ -181,7 +181,7 @@ public:
 	/**
 	* Remove Geometry and update dependent elements
 	*/
-	CHAOS_API virtual void RemoveElements(const FName & Group, const TArray<int32>& DeletionList, FProcessingParameters Params = FProcessingParameters()) override;
+	CHAOS_API virtual void RemoveElements(const FName & Group, const TArray<int32>& SortedDeletionList, FProcessingParameters Params = FProcessingParameters()) override;
 
 	/**
 	* Empty each managed array in each group in order to reset the collection to an initial (empty) state. 
