Chaosflesh: 5.6 node improvement
-optional muscle color update in ComputeFiberField
-speed up GenerateSurfaceBindings by only compute embedding in selected geometries
-deprecated Young's modulus and Damping parameter in the solver. They were already replaced by SetFleshDefaultProperties

diff --git a/Engine/Plugins/Experimental/ChaosFlesh/Source/ChaosFleshEngine/Private/ChaosFlesh/ChaosDeformableSolverComponent.cpp b/Engine/Plugins/Experimental/ChaosFlesh/Source/ChaosFleshEngine/Private/ChaosFlesh/ChaosDeformableSolverComponent.cpp
--- a/Engine/Plugins/Experimental/ChaosFlesh/Source/ChaosFleshEngine/Private/ChaosFlesh/ChaosDeformableSolverComponent.cpp
+++ b/Engine/Plugins/Experimental/ChaosFlesh/Source/ChaosFleshEngine/Private/ChaosFlesh/ChaosDeformableSolverComponent.cpp
@@ -235,10 +235,8 @@ void UDeformableSolverComponent::ResetSimulationProperties(const FSolverTimingGr
 			, false /*SolverCollisions.SolverGridBasedCollisions.bUseGridBasedConstraints*/
 			, 25. /*SolverCollisions.SolverGridBasedCollisions.GridDx*/
 			, SolverEvolution.SolverQuasistatics.bDoQuasistatics
-			, SolverForces.YoungModulus
 			, SolverConstraints.CorotatedConstraints.bDoBlended
 			, SolverConstraints.CorotatedConstraints.BlendedZeta
-			, SolverForces.Damping
 			, SolverForces.bEnableGravity
 			, SolverConstraints.CorotatedConstraints.bEnableCorotatedConstraint
 			, SolverConstraints.bEnablePositionTargets
@@ -280,10 +278,8 @@ void UDeformableSolverComponent::BuildSimulationProxy()
 			, false /*SolverCollisions.SolverGridBasedCollisions.bUseGridBasedConstraints*/
 			, 25. /*SolverCollisions.SolverGridBasedCollisions.GridDx*/
 			, SolverEvolution.SolverQuasistatics.bDoQuasistatics
-			, SolverForces.YoungModulus
 			, SolverConstraints.CorotatedConstraints.bDoBlended
 			, SolverConstraints.CorotatedConstraints.BlendedZeta
-			, SolverForces.Damping
 			, SolverForces.bEnableGravity
 			, SolverConstraints.CorotatedConstraints.bEnableCorotatedConstraint
 			, SolverConstraints.bEnablePositionTargets
diff --git a/Engine/Plugins/Experimental/ChaosFlesh/Source/ChaosFleshEngine/Public/ChaosFlesh/ChaosDeformableSolverGroups.h b/Engine/Plugins/Experimental/ChaosFlesh/Source/ChaosFleshEngine/Public/ChaosFlesh/ChaosDeformableSolverGroups.h
--- a/Engine/Plugins/Experimental/ChaosFlesh/Source/ChaosFleshEngine/Public/ChaosFlesh/ChaosDeformableSolverGroups.h
+++ b/Engine/Plugins/Experimental/ChaosFlesh/Source/ChaosFleshEngine/Public/ChaosFlesh/ChaosDeformableSolverGroups.h
@@ -1,4 +1,4 @@
-﻿// Copyright Epic Games, Inc. All Rights Reserved.
+// Copyright Epic Games, Inc. All Rights Reserved.
 #pragma once
 
 #include "ChaosFlesh/ChaosDeformableTypes.h"
@@ -217,14 +217,24 @@ struct FSolverForcesGroup
 {
 	GENERATED_BODY()
 
-	UPROPERTY(EditAnywhere, Category = "Forces")
+	UE_DEPRECATED(5.6, "YoungModulus has been deprecated.")
+	UPROPERTY(meta = (DeprecatedProperty, DeprecationMessage = "Please use SetFleshDefaultProperties dataflow node to set up vertex stiffness."))
 	float YoungModulus = 100000;
 
-	UPROPERTY(EditAnywhere, Category = "Forces")
+	UE_DEPRECATED(5.6, "Damping has been deprecated.")
+	UPROPERTY(meta = (DeprecatedProperty, DeprecationMessage = "Please use SetFleshDefaultProperties dataflow node to set up vertex damping."))
 	float Damping = 0;
 
 	UPROPERTY(EditAnywhere, Category = "Forces")
 	bool bEnableGravity = true;
+
+	PRAGMA_DISABLE_DEPRECATION_WARNINGS
+	FSolverForcesGroup() = default;
+	FSolverForcesGroup(const FSolverForcesGroup&) = default;
+	FSolverForcesGroup(FSolverForcesGroup&&) = default;
+	FSolverForcesGroup& operator=(const FSolverForcesGroup&) = default;
+	FSolverForcesGroup& operator=(FSolverForcesGroup&&) = default;
+	PRAGMA_ENABLE_DEPRECATION_WARNINGS
 };
 
 USTRUCT(BlueprintType)
diff --git a/Engine/Plugins/Experimental/ChaosFlesh/Source/ChaosFleshNodes/Private/Dataflow/ChaosFleshComputeFiberFieldNode.cpp b/Engine/Plugins/Experimental/ChaosFlesh/Source/ChaosFleshNodes/Private/Dataflow/ChaosFleshComputeFiberFieldNode.cpp
--- a/Engine/Plugins/Experimental/ChaosFlesh/Source/ChaosFleshNodes/Private/Dataflow/ChaosFleshComputeFiberFieldNode.cpp
+++ b/Engine/Plugins/Experimental/ChaosFlesh/Source/ChaosFleshNodes/Private/Dataflow/ChaosFleshComputeFiberFieldNode.cpp
@@ -138,28 +138,21 @@ void FComputeFiberFieldNode::Evaluate(UE::Dataflow::FContext& Context, const FDa
 		// Set output(s)
 		//
 
-		TManagedArray<FVector3f>* FiberDirections =
-			InCollection.FindAttribute<FVector3f>("FiberDirection", FTetrahedralCollection::TetrahedralGroup);
-		if (!FiberDirections)
-		{
-			FiberDirections =
-				&InCollection.AddAttribute<FVector3f>("FiberDirection", FTetrahedralCollection::TetrahedralGroup);
-		}
-		(*FiberDirections) = MoveTemp(FiberDirs);
+		TManagedArray<FVector3f>& FiberDirections =
+			InCollection.AddAttribute<FVector3f>("FiberDirection", FTetrahedralCollection::TetrahedralGroup);
+		FiberDirections = MoveTemp(FiberDirs);
 
-		TManagedArray<FLinearColor>* Color =
-			InCollection.FindAttribute<FLinearColor>("Color", FGeometryCollection::VerticesGroup);
-		if (!Color)
-		{
-			Color =
-				&InCollection.AddAttribute<FLinearColor>("Color", FGeometryCollection::VerticesGroup);
-		}
-		for (int32 i = 0; i < Color->Num(); ++i)
+		if (bShowMuscleColor)
 		{
-			float s = MuscleAttachmentScalarFieldTArray[i];
-			if (s > 0) // 1 <= s <= 2 if muscle
+			TManagedArray<FLinearColor>& Color =
+				InCollection.AddAttribute<FLinearColor>("Color", FGeometryCollection::VerticesGroup);
+			for (int32 i = 0; i < Color.Num(); ++i)
 			{
-				(*Color)[i] = FLinearColor(FVector(s-1, 0, 2 - s));
+				float Value = MuscleAttachmentScalarFieldTArray[i];
+				if (Value >= 1) // 1 <= Value <= 2 if muscle with origin and insertion
+				{
+					Color[i] = FLinearColor(FVector(Value - 1, 0, 2 - Value));
+				}
 			}
 		}
 		Out->SetValue(MoveTemp(InCollection), Context);
diff --git a/Engine/Plugins/Experimental/ChaosFlesh/Source/ChaosFleshNodes/Private/Dataflow/ChaosFleshGenerateSurfaceBindingsNode.cpp b/Engine/Plugins/Experimental/ChaosFlesh/Source/ChaosFleshNodes/Private/Dataflow/ChaosFleshGenerateSurfaceBindingsNode.cpp
--- a/Engine/Plugins/Experimental/ChaosFlesh/Source/ChaosFleshNodes/Private/Dataflow/ChaosFleshGenerateSurfaceBindingsNode.cpp
+++ b/Engine/Plugins/Experimental/ChaosFlesh/Source/ChaosFleshNodes/Private/Dataflow/ChaosFleshGenerateSurfaceBindingsNode.cpp
@@ -300,6 +300,9 @@ FGenerateSurfaceBindings::Evaluate(UE::Dataflow::FContext& Context, const FDataf
 					GeometryGroupSelected.Add(TetMeshIdx);
 				}
 			}
+			TArray<int32> UsedTetsIndexToGlobalTetIndex;
+			TArray<int32> GlobalTetIndexToUsedTetsIndex;
+			GlobalTetIndexToUsedTetsIndex.Init(INDEX_NONE, Tetrahedron->Num());
 			for (const int32& TetMeshIdx: GeometryGroupSelected)
 			{
 				if (GeometryGroupGuidsLocal.Num() && Guids)
@@ -317,11 +320,13 @@ FGenerateSurfaceBindings::Evaluate(UE::Dataflow::FContext& Context, const FDataf
 					{
 						const int32 Idx = TetMeshStart + i;
 						const FIntVector4& Tet = (*Tetrahedron)[Idx];
+						GlobalTetIndexToUsedTetsIndex[Idx] = Tets.Num();
 						Tets.Add(Chaos::TTetrahedron<Chaos::FReal>(
 							(*Vertex)[Tet[0]],
 							(*Vertex)[Tet[1]],
 							(*Vertex)[Tet[2]],
 							(*Vertex)[Tet[3]]));
+						UsedTetsIndexToGlobalTetIndex.Add(Idx);
 					}
 
 					if (bDoSurfaceProjection)
@@ -391,20 +396,29 @@ FGenerateSurfaceBindings::Evaluate(UE::Dataflow::FContext& Context, const FDataf
 						const int32 TetIdx = TetIntersections[j];
 						if (!Tets[TetIdx].Outside(Pos, 0.f)) // includes boundary
 						{
-							TetHits++;
-							Parents[LOD][i] = (*Tetrahedron)[TetIdx];
 							Chaos::TVector<Chaos::FReal, 4> WeightsD = Tets[TetIdx].GetBarycentricCoordinates(Pos);
-							Weights[LOD][i] = FVector4f(WeightsD[0], WeightsD[1], WeightsD[2], WeightsD[3]);
-							Offsets[LOD][i] = FVector3f(0);
-							Masks[LOD][i] = 1.0; // Shader does sim for this vertex
-
+							const int32 GlobalTetIndex = UsedTetsIndexToGlobalTetIndex[TetIdx];
 							FVector3f EmbeddedPos =
-								(*Vertex)[Parents[LOD][i][0]] * Weights[LOD][i][0] +
-								(*Vertex)[Parents[LOD][i][1]] * Weights[LOD][i][1] +
-								(*Vertex)[Parents[LOD][i][2]] * Weights[LOD][i][2] +
-								(*Vertex)[Parents[LOD][i][3]] * Weights[LOD][i][3];
-							check((Pos - EmbeddedPos).SquaredLength() < UE_SMALL_NUMBER);
-							break;
+								(*Vertex)[(*Tetrahedron)[GlobalTetIndex][0]] * WeightsD[0] +
+								(*Vertex)[(*Tetrahedron)[GlobalTetIndex][1]] * WeightsD[1] +
+								(*Vertex)[(*Tetrahedron)[GlobalTetIndex][2]] * WeightsD[2] +
+								(*Vertex)[(*Tetrahedron)[GlobalTetIndex][3]] * WeightsD[3];
+							if ((Pos - EmbeddedPos).SquaredLength() < UE_SMALL_NUMBER)
+							{
+								TetHits++;
+								Parents[LOD][i] = (*Tetrahedron)[GlobalTetIndex];
+								Weights[LOD][i] = FVector4f(WeightsD[0], WeightsD[1], WeightsD[2], WeightsD[3]);
+								Offsets[LOD][i] = FVector3f(0);
+								Masks[LOD][i] = 1.0; // Shader does sim for this vertex
+								break;
+							}
+							else
+							{
+								UE_LOG(LogMeshBindings, Warning,
+									TEXT("Vertex position does not collocate with interpolated position in tetrahedron %d, distance = %.4f)"),
+									TetIdx,
+									(Pos - EmbeddedPos).Length());
+							}
 						}
 					}
 					if (j == TetIntersections.Num())
@@ -550,9 +564,10 @@ FGenerateSurfaceBindings::Evaluate(UE::Dataflow::FContext& Context, const FDataf
 									for (int32 k = 0; k < NeighborTets.Num(); k++)
 									{
 										const int32 TetIdx = NeighborTets[k];
-										if (ensure(Tets.IsValidIndex(TetIdx)))
+										const int32 UsedTetIdx = GlobalTetIndexToUsedTetsIndex[TetIdx];
+										if (ensure(Tets.IsValidIndex(UsedTetIdx)))
 										{
-											const Chaos::FTetrahedron& Tet = Tets[TetIdx];
+											const Chaos::FTetrahedron& Tet = Tets[UsedTetIdx];
 
 											Chaos::TVec4<Chaos::FReal> W;
 											Chaos::TVec3<Chaos::FReal> EmbeddedPos = Tet.FindClosestPointAndBary(PosD, W, 1.0e-4);
diff --git a/Engine/Plugins/Experimental/ChaosFlesh/Source/ChaosFleshNodes/Public/Dataflow/ChaosFleshComputeFiberFieldNode.h b/Engine/Plugins/Experimental/ChaosFlesh/Source/ChaosFleshNodes/Public/Dataflow/ChaosFleshComputeFiberFieldNode.h
--- a/Engine/Plugins/Experimental/ChaosFlesh/Source/ChaosFleshNodes/Public/Dataflow/ChaosFleshComputeFiberFieldNode.h
+++ b/Engine/Plugins/Experimental/ChaosFlesh/Source/ChaosFleshNodes/Public/Dataflow/ChaosFleshComputeFiberFieldNode.h
@@ -50,6 +50,10 @@ public:
 	UPROPERTY(EditAnywhere, Category = "Dataflow")
 	float Tolerance = 1.0e-7;
 
+	//If muscles are colored by the flow from origins (blue) to insertions (red). This becomes optional in 5.6
+	UPROPERTY(EditAnywhere, Category = "Dataflow")
+	bool bShowMuscleColor = true;
+
 	FComputeFiberFieldNode(const UE::Dataflow::FNodeParameters& InParam, FGuid InGuid = FGuid::NewGuid())
 		: FDataflowNode(InParam, InGuid)
 	{
diff --git a/Engine/Plugins/Experimental/ChaosFlesh/Source/ChaosFleshNodes/Public/Dataflow/ChaosFleshSetVertexTrianglePositionTargetBindingNode.h b/Engine/Plugins/Experimental/ChaosFlesh/Source/ChaosFleshNodes/Public/Dataflow/ChaosFleshSetVertexTrianglePositionTargetBindingNode.h
--- a/Engine/Plugins/Experimental/ChaosFlesh/Source/ChaosFleshNodes/Public/Dataflow/ChaosFleshSetVertexTrianglePositionTargetBindingNode.h
+++ b/Engine/Plugins/Experimental/ChaosFlesh/Source/ChaosFleshNodes/Public/Dataflow/ChaosFleshSetVertexTrianglePositionTargetBindingNode.h
@@ -9,6 +9,9 @@
 
 #include "ChaosFleshSetVertexTrianglePositionTargetBindingNode.generated.h"
 
+/**
+ * Create point-triangle weak constraints (springs) between different geometries based on search radius.
+ */
 USTRUCT(meta = (DataflowFlesh))
 struct FSetVertexTrianglePositionTargetBindingDataflowNode : public FDataflowNode
 {
@@ -26,13 +29,14 @@ public:
 	UPROPERTY(EditAnywhere, Category = "Dataflow")
 		float PositionTargetStiffness = 10000.f;
 
+	/* (optional) only create weak constraints from vertices in VertexSelection to triangles in other geometries */
 	UPROPERTY(meta = (DataflowInput, DisplayName = "(Optional) VertexSelection"))
 		FDataflowVertexSelection VertexSelection;
-
+	
+	/* weak constraint search radius = Largest side of the collection's bounding box * VertexRadiusRatio */
 	UPROPERTY(EditAnywhere, Category = "Dataflow", meta = (ClampMin = "0.0", ClampMax = "2.0", UIMin = "0.0", UIMax = "2.0"))
 		float VertexRadiusRatio = .001f;
 
-
 	FSetVertexTrianglePositionTargetBindingDataflowNode(const UE::Dataflow::FNodeParameters& InParam, FGuid InGuid = FGuid::NewGuid())
 		: FDataflowNode(InParam, InGuid)
 	{
diff --git a/Engine/Plugins/Experimental/GeometryCollectionPlugin/Source/GeometryCollectionNodes/Public/Dataflow/GeometryCollectionTransferVertexAttributeNode.h b/Engine/Plugins/Experimental/GeometryCollectionPlugin/Source/GeometryCollectionNodes/Public/Dataflow/GeometryCollectionTransferVertexAttributeNode.h
--- a/Engine/Plugins/Experimental/GeometryCollectionPlugin/Source/GeometryCollectionNodes/Public/Dataflow/GeometryCollectionTransferVertexAttributeNode.h
+++ b/Engine/Plugins/Experimental/GeometryCollectionPlugin/Source/GeometryCollectionNodes/Public/Dataflow/GeometryCollectionTransferVertexAttributeNode.h
@@ -102,7 +102,9 @@ public:
 	EDataflowTransferVertexAttributeNodeFalloff Falloff = EDataflowTransferVertexAttributeNodeFalloff::None;
 
 	/* Threshold based on distance from source triangle.Values past the threshold will falloff.[Defaults to 1 percent of triangle size(0.01)] */
-	UPROPERTY(EditAnywhere, Category = "Thresholds", meta = (EditCondition = "BoundingVolumeType == EDataflowTransferVertexAttributeNodeBoundingVolume::Triangle", EditConditionHides))
+	UPROPERTY(EditAnywhere, Category = "Thresholds", meta = (EditCondition = 
+		"BoundingVolumeType == EDataflowTransferVertexAttributeNodeBoundingVolume::Triangle && EDataflowTransferVertexAttributeNodeFalloff Falloff = EDataflowTransferVertexAttributeNodeFalloff::None"
+		, EditConditionHides))
 	float FalloffThreshold = 0.01f;
 
 	/* Edge multiplier for the Bounding Volume Hierarchy(BVH) target's particle search radius. */
diff --git a/Engine/Source/Runtime/Experimental/Chaos/Private/Chaos/Deformable/ChaosDeformableSolver.cpp b/Engine/Source/Runtime/Experimental/Chaos/Private/Chaos/Deformable/ChaosDeformableSolver.cpp
--- a/Engine/Source/Runtime/Experimental/Chaos/Private/Chaos/Deformable/ChaosDeformableSolver.cpp
+++ b/Engine/Source/Runtime/Experimental/Chaos/Private/Chaos/Deformable/ChaosDeformableSolver.cpp
@@ -108,8 +108,8 @@ namespace Chaos::Softs
 	int32 GSParallelMax = 100;
 	FAutoConsoleVariableRef CVarDeformableGSParrelMax(TEXT("p.Chaos.Deformable.GSParallelMax"), GSParallelMax, TEXT("Minimal number of particles to process in parallel for Gauss Seidel constraints. [def: 100]"));
 
-	float MaxDxRatio = .01f;
-	FAutoConsoleVariableRef CVarDeformableGSMaxDxRatio(TEXT("p.Chaos.Deformable.GSMaxDxRatio"), MaxDxRatio, TEXT("Max size for dx in each iteration for Gauss Seidel constraints. [def: .01]"));
+	float MaxDxRatio = 1.f;
+	FAutoConsoleVariableRef CVarDeformableGSMaxDxRatio(TEXT("p.Chaos.Deformable.GSMaxDxRatio"), MaxDxRatio, TEXT("Max size for dx in each iteration for Gauss Seidel constraints. [def: 1]"));
 
 	FDeformableSolver::FDeformableSolver(FDeformableSolverProperties InProp)
 		: CurrentInputPackage(TUniquePtr < FDeformablePackage >(nullptr))
@@ -146,7 +146,7 @@ namespace Chaos::Softs
 			Property.NumSolverIterations, (FSolverReal)0.,
 			/*SelfCollisionsThickness = */(FSolverReal)0.,
 			/*CoefficientOfFriction = */(FSolverReal)0.,
-			/*FSolverReal Damping = */Property.Damping,
+			/*FSolverReal Damping = */(FSolverReal)0.,
 			/*FSolverReal LocalDamping = */(FSolverReal)0.,
 			Property.bDoQuasistatics, 
 			true));
@@ -205,7 +205,6 @@ namespace Chaos::Softs
 	void FDeformableSolver::Simulate(FSolverReal DeltaTime)
 	{
 		PERF_SCOPE(STAT_ChaosDeformableSolver_Simulate);
-
 		if (Property.NumSolverIterations)
 		{
 			RemoveSimulationObjects();
@@ -950,7 +949,7 @@ namespace Chaos::Softs
 				}
 				const TManagedArray<FSolverReal>* StiffnessArray = Rest.FindAttribute<FSolverReal>("Stiffness", FGeometryCollection::VerticesGroup);
 				TArray<FSolverReal> TetStiffness;
-				TetStiffness.Init(Property.EMesh, Elements.Num());
+				TetStiffness.Init(0.f, Elements.Num());
 				if (StiffnessArray)
 				{
 					for (int32 edx = 0; edx < Elements.Num(); edx++)
@@ -1149,7 +1148,7 @@ namespace Chaos::Softs
 				AllCorotatedCodEMeshArray->SetNum(AllCorotatedCodEMeshArray->Num() + AllUnconstrainedSurfaceElementsCorotatedCodRange[1]);
 				for (int32 i = AllUnconstrainedSurfaceElementsCorotatedCodRange[0]; i < AllUnconstrainedSurfaceElementsCorotatedCodRange[0] + AllUnconstrainedSurfaceElementsCorotatedCodRange[1]; i++)
 				{
-					(*AllCorotatedCodEMeshArray)[i] = Property.EMesh;
+					(*AllCorotatedCodEMeshArray)[i] = 0.f;
 				}
 			}
 		}
@@ -1160,7 +1159,7 @@ namespace Chaos::Softs
 				AllSkinEMeshArray->SetNum(AllSkinEMeshArray->Num() + AllUnconstrainedSurfaceElementsSkinRange[1]);
 				for (int32 i = AllUnconstrainedSurfaceElementsSkinRange[0]; i < AllUnconstrainedSurfaceElementsSkinRange[0] + AllUnconstrainedSurfaceElementsSkinRange[1]; i++)
 				{
-					(*AllSkinEMeshArray)[i] = Property.EMesh;
+					(*AllSkinEMeshArray)[i] = 0.f;
 				}
 			}
 		}
diff --git a/Engine/Source/Runtime/Experimental/Chaos/Public/Chaos/Deformable/ChaosDeformableSolverTypes.h b/Engine/Source/Runtime/Experimental/Chaos/Public/Chaos/Deformable/ChaosDeformableSolverTypes.h
--- a/Engine/Source/Runtime/Experimental/Chaos/Public/Chaos/Deformable/ChaosDeformableSolverTypes.h
+++ b/Engine/Source/Runtime/Experimental/Chaos/Public/Chaos/Deformable/ChaosDeformableSolverTypes.h
@@ -26,10 +26,8 @@ namespace Chaos::Softs
 			bool InbUseGridBasedConstraints = false,
 			FSolverReal InGridDx = (FSolverReal)1.,
 			bool InbDoQuasistatics = false,
-			FSolverReal InEMesh = (FSolverReal)100000.,
 			bool InbDoBlended = false,
 			FSolverReal InBlendedZeta = (FSolverReal).1,
-			FSolverReal InDamping = (FSolverReal)0,
 			bool InbEnableGravity = true,
 			bool InbEnableCorotatedConstraints = true,
 			bool InbEnablePositionTargets = true,
@@ -59,10 +57,8 @@ namespace Chaos::Softs
 			, bUseGridBasedConstraints(InbUseGridBasedConstraints)
 			, GridDx(InGridDx)
 			, bDoQuasistatics(InbDoQuasistatics)
-			, EMesh(InEMesh)
 			, bDoBlended(InbDoBlended)
 			, BlendedZeta(InBlendedZeta)
-			, Damping(InDamping)
 			, bEnableGravity(InbEnableGravity)
 			, bEnableCorotatedConstraints(InbEnableCorotatedConstraints)
 			, bEnablePositionTargets(InbEnablePositionTargets)
@@ -94,10 +90,8 @@ namespace Chaos::Softs
 		bool bUseGridBasedConstraints = false;
 		FSolverReal GridDx = (FSolverReal)1.;
 		bool bDoQuasistatics = false;
-		FSolverReal EMesh = (FSolverReal)100000.;
 		bool bDoBlended = false;
 		FSolverReal BlendedZeta = (FSolverReal)0.;
-		FSolverReal Damping = (FSolverReal)0.;
 		bool bEnableGravity = true;
 		bool bEnableCorotatedConstraints = true;
 		bool bEnablePositionTargets = true;
