Chaosflesh: CIS fix [Issue 798879, CL 40298875], constructor missing for struct

diff --git a/Engine/Plugins/Experimental/ChaosFlesh/Source/ChaosFleshNodes/Public/Dataflow/ChaosFleshComputeMuscleActivationNode.h b/Engine/Plugins/Experimental/ChaosFlesh/Source/ChaosFleshNodes/Public/Dataflow/ChaosFleshComputeMuscleActivationNode.h
--- a/Engine/Plugins/Experimental/ChaosFlesh/Source/ChaosFleshNodes/Public/Dataflow/ChaosFleshComputeMuscleActivationNode.h
+++ b/Engine/Plugins/Experimental/ChaosFlesh/Source/ChaosFleshNodes/Public/Dataflow/ChaosFleshComputeMuscleActivationNode.h
@@ -169,6 +169,10 @@ struct FCurveMuscleName
 	FString CurveName;
 	UPROPERTY(EditAnywhere, Category = "Muscle Activation Parameters")
 	FString MuscleName;
+
+	FCurveMuscleName() = default;
+	FCurveMuscleName(const FString& InCurveName, const FString& InMuscleName)
+		: CurveName(InCurveName), MuscleName(InMuscleName) {};
 };
 
 USTRUCT(meta = (DataflowFlesh))
