Chaosflesh:
-added near-zero volume check for degenerate elements
-use GetNormal() instead of custom normal calculation in TriangleMesh.PointClosestTriangleQuery()
-added spring mass ratio statistics in SetVertexTrianglePositionTargetBinding

diff --git a/Engine/Plugins/Experimental/ChaosFlesh/Source/ChaosFleshNodes/Private/Dataflow/ChaosFleshSetFleshDefaultPropertiesNode.cpp b/Engine/Plugins/Experimental/ChaosFlesh/Source/ChaosFleshNodes/Private/Dataflow/ChaosFleshSetFleshDefaultPropertiesNode.cpp
--- a/Engine/Plugins/Experimental/ChaosFlesh/Source/ChaosFleshNodes/Private/Dataflow/ChaosFleshSetFleshDefaultPropertiesNode.cpp
+++ b/Engine/Plugins/Experimental/ChaosFlesh/Source/ChaosFleshNodes/Private/Dataflow/ChaosFleshSetFleshDefaultPropertiesNode.cpp
@@ -65,6 +65,8 @@ void FSetFleshDefaultPropertiesNode::Evaluate(UE::Dataflow::FContext& Context, c
 				Mass.Fill(0.f);
 				float MinV = TNumericLimits<float>::Max();
 				float MaxV = -TNumericLimits<float>::Max();
+				int32 NegativeElementVolumeCount = 0;
+				int32 SmallElementVolumeCount = 0;
 				double AvgV = 0.0;
 				TSet<int32> Visited;
 				int32 NumSet = 0;
@@ -81,13 +83,36 @@ void FSetFleshDefaultPropertiesNode::Evaluate(UE::Dataflow::FContext& Context, c
 						FVector3f X1 = Vertex[Tetrahedron[e][1]];
 						FVector3f X2 = Vertex[Tetrahedron[e][2]];
 						FVector3f X3 = Vertex[Tetrahedron[e][3]];
-						ElementVolume[e] = ((X1 - X0).Dot(FVector3f::CrossProduct(X3 - X0, X2 - X0))) / 6.f;
+						ElementVolume[e] = ((X3 - X0).Dot(FVector3f::CrossProduct(X1 - X0, X2 - X0))) / 6.f;
 						if (ElementVolume[e] < 0.f)
 						{
 							ElementVolume[e] = -ElementVolume[e];
+							NegativeElementVolumeCount++;
+						}
+						if (FMath::Abs(ElementVolume[e]) < UE_SMALL_NUMBER)
+						{
+							SmallElementVolumeCount++;
+							if (SmallElementVolumeCount == 1)
+							{
+								UE_LOG(LogChaosFlesh, Error,
+									TEXT("'%s' - Example: tetrahedron %d has volume %f < %e."),
+									*GetName().ToString(), e, ElementVolume[e], UE_SMALL_NUMBER);
+							}
 						}
 						TotalVolume += ElementVolume[e];
 					}
+					if (NegativeElementVolumeCount)
+					{
+						UE_LOG(LogChaosFlesh, Warning,
+							TEXT("'%s' - Flipped negative volume for %d tetrahedra."),
+							*GetName().ToString(), NegativeElementVolumeCount);
+					}
+					if (SmallElementVolumeCount)
+					{
+						UE_LOG(LogChaosFlesh, Error,
+							TEXT("'%s' - %d tetrahedra have volume < %e."),
+							*GetName().ToString(), SmallElementVolumeCount, UE_SMALL_NUMBER);
+					}
 					for (int32 e = 0; e < TetsNum; e++)
 					{
 						ElementMass[e] = Density * ElementVolume[e];
diff --git a/Engine/Plugins/Experimental/ChaosFlesh/Source/ChaosFleshNodes/Private/Dataflow/ChaosFleshSetVertexTrianglePositionTargetBindingNode.cpp b/Engine/Plugins/Experimental/ChaosFlesh/Source/ChaosFleshNodes/Private/Dataflow/ChaosFleshSetVertexTrianglePositionTargetBindingNode.cpp
--- a/Engine/Plugins/Experimental/ChaosFlesh/Source/ChaosFleshNodes/Private/Dataflow/ChaosFleshSetVertexTrianglePositionTargetBindingNode.cpp
+++ b/Engine/Plugins/Experimental/ChaosFlesh/Source/ChaosFleshNodes/Private/Dataflow/ChaosFleshSetVertexTrianglePositionTargetBindingNode.cpp
@@ -5,11 +5,12 @@
 #include "Chaos/Utilities.h"
 #include "Chaos/HierarchicalSpatialHash.h"
 #include "Chaos/TriangleCollisionPoint.h"
+#include "ChaosFlesh/ChaosFlesh.h"
+#include "ChaosFlesh/ChaosFleshCollectionFacade.h"
+#include "ChaosFlesh/TetrahedralCollection.h"
 #include "GeometryCollection/Facades/CollectionMeshFacade.h"
 #include "GeometryCollection/Facades/CollectionPositionTargetFacade.h"
 #include "GeometryCollection/Facades/CollectionCollisionFacade.h"
-#include "ChaosFlesh/TetrahedralCollection.h"
-#include "ChaosFlesh/ChaosFleshCollectionFacade.h"
 
 #include UE_INLINE_GENERATED_CPP_BY_NAME(ChaosFleshSetVertexTrianglePositionTargetBindingNode)
 
@@ -103,7 +104,9 @@ void FSetVertexTrianglePositionTargetBindingDataflowNode::Evaluate(UE::Dataflow:
 
 					GeometryCollection::Facades::FPositionTargetFacade PositionTargets(InCollection);
 					PositionTargets.DefineSchema();
-
+#if WITH_EDITOR
+					TArray<float> MassRatioArray;
+#endif 
 					TArray<Chaos::TVec3<Chaos::FReal>> VertexTVec3;
 					VertexTVec3.SetNum(Vertex.Num());
 					for (int32 VertexIdx = 0; VertexIdx < Vertex.Num(); VertexIdx++)
@@ -141,11 +144,18 @@ void FSetVertexTrianglePositionTargetBindingDataflowNode::Evaluate(UE::Dataflow:
 								DataPackage.SourceWeights[2] = CollisionPoint.Bary[3];
 								if (TManagedArray<float>* Mass = InCollection.FindAttribute<float>("Mass", FGeometryCollection::VerticesGroup))
 								{
+									float MinMass = (*Mass)[DataPackage.TargetIndex[0]];
+									float MaxMass = (*Mass)[DataPackage.TargetIndex[0]];
 									DataPackage.Stiffness = 0.f;
 									for (int32 k = 0; k < 3; k++)
 									{
+										MinMass = FGenericPlatformMath::Min(MinMass, (*Mass)[DataPackage.SourceIndex[k]]);
+										MaxMass = FGenericPlatformMath::Max(MaxMass, (*Mass)[DataPackage.SourceIndex[k]]);
 										DataPackage.Stiffness += DataPackage.SourceWeights[k] * PositionTargetStiffness * (*Mass)[DataPackage.SourceIndex[k]];
 									}
+#if WITH_EDITOR
+									MassRatioArray.Add(MaxMass / MinMass);
+#endif
 									DataPackage.Stiffness += DataPackage.TargetWeights[0] * PositionTargetStiffness * (*Mass)[DataPackage.TargetIndex[0]];
 								}
 								else
@@ -156,6 +166,14 @@ void FSetVertexTrianglePositionTargetBindingDataflowNode::Evaluate(UE::Dataflow:
 							}
 						}
 					}
+#if WITH_EDITOR
+					MassRatioArray.Sort();
+					if (MassRatioArray.Num())
+					{
+						UE_LOG(LogChaosFlesh, Warning, TEXT("SetVertexTrianglePositionTargetBinding: Max mass ratio = %f, median mass ratio = %f")
+							, MassRatioArray[MassRatioArray.Num()-1], MassRatioArray[MassRatioArray.Num()/2]);
+					}
+#endif
 				}
 			}
 		}
diff --git a/Engine/Source/Runtime/Experimental/Chaos/Private/Chaos/TriangleMesh.cpp b/Engine/Source/Runtime/Experimental/Chaos/Private/Chaos/TriangleMesh.cpp
--- a/Engine/Source/Runtime/Experimental/Chaos/Private/Chaos/TriangleMesh.cpp
+++ b/Engine/Source/Runtime/Experimental/Chaos/Private/Chaos/TriangleMesh.cpp
@@ -1936,11 +1936,12 @@ bool FTriangleMesh::PointClosestTriangleQuery(const TSpatialHashType<T>& Spatial
 		const TVec3<T>& A = Points[MElements[MinTriId][0]];
 		const TVec3<T>& B = Points[MElements[MinTriId][1]];
 		const TVec3<T>& C = Points[MElements[MinTriId][2]];
+		const TTriangle<T> Triangle(A, B, C);
 		TVec3<T> Bary;
 		const TVec3<T> ClosestPoint = FindClosestPointAndBaryOnTriangle(A, B, C, PointPosition, Bary);
 		const T DistSq = (PointPosition - ClosestPoint).SizeSquared();
-		TVec3<T> Normal = TVec3<T>::CrossProduct(B - A, C - A).GetSafeNormal();	
-
+		const TVec3<T> Normal = Triangle.GetNormal();
+		
 		TTriangleCollisionPoint<T> CollisionPoint;
 		CollisionPoint.ContactType = TTriangleCollisionPoint<T>::EContactType::PointFace;
 		CollisionPoint.Indices[0] = PointIndex;
@@ -1949,7 +1950,7 @@ bool FTriangleMesh::PointClosestTriangleQuery(const TSpatialHashType<T>& Spatial
 		CollisionPoint.Location = ClosestPoint;
 		CollisionPoint.Normal = Normal;
 		CollisionPoint.Phi = FMath::Sqrt(DistSq);
-		if (TVec3<T>::DotProduct(Normal, PointPosition - A) > 0) //Point is inside of triangle
+		if (TVec3<T>::DotProduct(Normal, PointPosition - A) > 0) //Point is inside of triangle (normal is pointing inward)
 		{
 			CollisionPoint.Phi = -CollisionPoint.Phi;
 		}
