Chaosflesh:
-Added AnimationAssetTerminalNode to generate training animation for activation MLD on curves.
-FReadSkeletalMeshCurvesDataflowNode returns a GeometrySelection for visual debugging

diff --git a/Engine/Plugins/Experimental/ChaosFlesh/Source/ChaosFleshNodes/Private/Dataflow/ChaosFleshComputeMuscleActivationNode.cpp b/Engine/Plugins/Experimental/ChaosFlesh/Source/ChaosFleshNodes/Private/Dataflow/ChaosFleshComputeMuscleActivationNode.cpp
--- a/Engine/Plugins/Experimental/ChaosFlesh/Source/ChaosFleshNodes/Private/Dataflow/ChaosFleshComputeMuscleActivationNode.cpp
+++ b/Engine/Plugins/Experimental/ChaosFlesh/Source/ChaosFleshNodes/Private/Dataflow/ChaosFleshComputeMuscleActivationNode.cpp
@@ -6,6 +6,8 @@
 #include "GeometryCollection/Facades/CollectionMuscleActivationFacade.h"
 #include "ChaosFlesh/TetrahedralCollection.h"
 #include "Engine/SkeletalMesh.h"
+#include "Dataflow/DataflowSelection.h"
+#include "ChaosFlesh/ChaosFlesh.h"
 
 #include UE_INLINE_GENERATED_CPP_BY_NAME(ChaosFleshComputeMuscleActivationNode)
 
@@ -143,6 +145,11 @@ FReadSkeletalMeshCurvesDataflowNode::FReadSkeletalMeshCurvesDataflowNode(const U
 						{
 							TArray<FName> CurveNamesArray;
 							AnimCurveMetaData->GetCurveMetaDataNames(CurveNamesArray);
+							CurveNamesArray.Sort([](const FName& A, const FName& B)
+							{
+								return A.ToString() < B.ToString();
+							});
+
 							for (FName Curvename : CurveNamesArray)
 							{
 								CurveMuscleNameArray.Add(FCurveMuscleName(Curvename.ToString(), FString()));
@@ -160,22 +167,38 @@ FReadSkeletalMeshCurvesDataflowNode::FReadSkeletalMeshCurvesDataflowNode(const U
 	RegisterInputConnection(&Collection);
 	RegisterInputConnection(&SkeletalMesh);
 	RegisterOutputConnection(&Collection, &Collection);
+	RegisterOutputConnection(&GeometrySelection);
 }
 
 void FReadSkeletalMeshCurvesDataflowNode::Evaluate(UE::Dataflow::FContext& Context, const FDataflowOutput* Out) const
 {
-	if (Out->IsA(&Collection))
+	if (Out->IsA(&Collection) || Out->IsA(&GeometrySelection))
 	{
 		FManagedArrayCollection InCollection = GetValue(Context, &Collection);
 		GeometryCollection::Facades::FMuscleActivationFacade MuscleActivationFacade(InCollection);
+		TArray<int32> MuscleCurveGeometry;
 		if (MuscleActivationFacade.IsValid())
 		{
 			for (const FCurveMuscleName& CurveMuscleName : CurveMuscleNameArray)
 			{
-				MuscleActivationFacade.AssignCurveName(CurveMuscleName.CurveName, CurveMuscleName.MuscleName);
+				const int32 MuscleIdx = MuscleActivationFacade.AssignCurveName(CurveMuscleName.CurveName, CurveMuscleName.MuscleName);
+				if (MuscleIdx != INDEX_NONE)
+				{
+					MuscleCurveGeometry.Add(MuscleActivationFacade.FindMuscleGeometryIndex(MuscleIdx));
+				}
+				else
+				{
+					UE_LOG(LogChaosFlesh, Error,
+						TEXT("ReadSkeletalMeshCurves: Geometry %s (connecting to curve %s) is not an active muscle."),
+						*CurveMuscleName.MuscleName, *CurveMuscleName.CurveName);
+				}
 			}
 		}
-		Out->SetValue(MoveTemp(InCollection), Context);
+		FDataflowGeometrySelection OutGeometrySelection;
+		OutGeometrySelection.Initialize(InCollection.NumElements(FGeometryCollection::GeometryGroup), false);
+		OutGeometrySelection.SetFromArray(MuscleCurveGeometry);
+		SetValue(Context, MoveTemp(InCollection), &Collection);
+		SetValue(Context, MoveTemp(OutGeometrySelection), &GeometrySelection);
 	}
 }
 
diff --git a/Engine/Plugins/Experimental/ChaosFlesh/Source/ChaosFleshNodes/Private/Dataflow/ChaosFleshFleshAssetTerminalNode.cpp b/Engine/Plugins/Experimental/ChaosFlesh/Source/ChaosFleshNodes/Private/Dataflow/ChaosFleshFleshAssetTerminalNode.cpp
--- a/Engine/Plugins/Experimental/ChaosFlesh/Source/ChaosFleshNodes/Private/Dataflow/ChaosFleshFleshAssetTerminalNode.cpp
+++ b/Engine/Plugins/Experimental/ChaosFlesh/Source/ChaosFleshNodes/Private/Dataflow/ChaosFleshFleshAssetTerminalNode.cpp
@@ -1,6 +1,12 @@
 // Copyright Epic Games, Inc. All Rights Reserved.
 
 #include "Dataflow/ChaosFleshFleshAssetTerminalNode.h"
+#include "Animation/Skeleton.h"
+#include "AnimationUtils.h"
+#include "Engine/SkeletalMesh.h"
+#include "Misc/PackageName.h"
+#include "UObject/Package.h"
+#include "UObject/SavePackage.h"
 
 #include UE_INLINE_GENERATED_CPP_BY_NAME(ChaosFleshFleshAssetTerminalNode)
 
@@ -19,3 +25,190 @@ void FFleshAssetTerminalDataflowNode::Evaluate(UE::Dataflow::FContext& Context,
 	const FManagedArrayCollection& InCollection = GetValue<FManagedArrayCollection>(Context, &Collection);
 	SetValue(Context, InCollection, &Collection);
 }
+
+void FCurveSamplingAnimationAssetTerminalNode::Evaluate(UE::Dataflow::FContext& Context) const
+{
+	TObjectPtr<USkeletalMesh> InSkeletalMesh = GetValue(Context, &SkeletalMeshAsset);
+	TObjectPtr<UAnimSequence> InAnimationAsset = GetValue(Context, &AnimationAsset);
+}
+
+
+void FCurveSamplingAnimationAssetTerminalNode::SetAssetValue(TObjectPtr<UObject> Asset, UE::Dataflow::FContext& Context) const
+{
+	// See Engine\Source\Editor\SequenceRecorder\Private\AnimationRecorder.cpp for reference
+	if (TObjectPtr<USkeletalMesh> InSkeletalMesh = GetValue(Context, &SkeletalMeshAsset))
+	{
+		UAnimSequence* AssetToSet = Cast<UAnimSequence>(Asset.Get());
+		if (!AssetToSet)
+		{
+			// use the input instead
+			AssetToSet = GetValue(Context, &AnimationAsset);
+		}
+#if WITH_EDITOR
+		if (AssetToSet)
+		{
+			static bool bTransactRecording = false;
+
+			// set skeleton
+			AssetToSet->SetSkeleton(InSkeletalMesh->GetSkeleton());
+
+			if (!AssetToSet->BoneCompressionSettings)
+			{
+				AssetToSet->BoneCompressionSettings = FAnimationUtils::GetDefaultAnimationRecorderBoneCompressionSettings();
+			}
+			const int32 NumBones = InSkeletalMesh->GetRefSkeleton().GetNum();
+			TArray<FTransform> RestTransforms;
+			InSkeletalMesh->GetRefSkeleton().GetBoneAbsoluteTransforms(RestTransforms);
+
+			TArray<FName> CurveNamesArray;
+			for (const UAssetUserData* AssetUserData : *InSkeletalMesh->GetAssetUserDataArray())
+			{
+				if (const UAnimCurveMetaData* AnimCurveMetaData = Cast<UAnimCurveMetaData>(AssetUserData))
+				{
+					AnimCurveMetaData->GetCurveMetaDataNames(CurveNamesArray);
+				}
+			}
+			const int32 NumCurves = CurveNamesArray.Num();
+			IAnimationDataController& Controller = AssetToSet->GetController();
+			Controller.SetModel(AssetToSet->GetDataModelInterface());
+			Controller.InitializeModel();
+			Controller.RemoveAllCurvesOfType(ERawCurveTrackTypes::RCT_Float);
+			Controller.RemoveAllBoneTracks(bTransactRecording);
+			AssetToSet->ResetAnimation();
+			TArray<FRawAnimSequenceTrack> RawTracks;
+
+			USkeleton* AnimSkeleton = AssetToSet->GetSkeleton();
+			// add all frames
+			for (int32 BoneIndex = 0; BoneIndex < NumBones; ++BoneIndex)
+			{
+				// verify if this bone exists in skeleton
+				const int32 BoneTreeIndex = AnimSkeleton->GetSkeletonBoneIndexFromMeshBoneIndex(
+					InSkeletalMesh, BoneIndex);
+				if (BoneTreeIndex != INDEX_NONE)
+				{
+					// add tracks for the bone existing
+					const FName BoneTreeName = AnimSkeleton->GetReferenceSkeleton().GetBoneName(BoneTreeIndex);
+					Controller.AddBoneCurve(BoneTreeName, bTransactRecording);
+					RawTracks.AddDefaulted();
+				}
+			}
+
+			AssetToSet->RetargetSource = AnimSkeleton->GetRetargetSourceForMesh(InSkeletalMesh);
+
+			// record transforms
+
+			TArray<FName> TrackNames;
+			const IAnimationDataModel* DataModel = AssetToSet->GetDataModel();
+			DataModel->GetBoneTrackNames(TrackNames);
+
+			//FSerializedAnimation SerializedAnimation;
+			for (int32 TrackIndex = 0; TrackIndex < TrackNames.Num(); ++TrackIndex)
+			{
+				const FName& TrackName = TrackNames[TrackIndex];
+				FRawAnimSequenceTrack& RawTrack = RawTracks[TrackIndex];
+
+				// verify if this bone exists in skeleton
+				const int32 BoneTreeIndex = AnimSkeleton->GetReferenceSkeleton().FindBoneIndex(TrackName);
+				if (BoneTreeIndex != INDEX_NONE)
+				{
+					const int32 BoneIndex = AnimSkeleton->GetMeshBoneIndexFromSkeletonBoneIndex(InSkeletalMesh, BoneTreeIndex);
+					const int32 ParentIndex = InSkeletalMesh->GetRefSkeleton().GetParentIndex(BoneIndex);
+					// Only record the rest pose for activation MLD training
+					FTransform LocalTransform = RestTransforms[BoneIndex];
+					if (ParentIndex != INDEX_NONE)
+					{
+						LocalTransform.SetToRelativeTransform(RestTransforms[ParentIndex]);
+					}
+
+					RawTrack.PosKeys.Add(FVector3f(LocalTransform.GetTranslation()));
+					RawTrack.RotKeys.Add(FQuat4f(LocalTransform.GetRotation()));
+					RawTrack.ScaleKeys.Add(FVector3f(LocalTransform.GetScale3D()));
+				}
+			}
+
+			const int32 NumFrames = NumFramesPerMuscle * NumCurves;
+
+			//Set Interpolation type (Step or Linear), doesn't look like there is a controller for this.
+			AssetToSet->Interpolation = EAnimInterpolationType::Linear;
+			const FFrameRate RecordingRate = FFrameRate(FrameRate, 1);
+			// Set frame rate and number of frames
+			Controller.SetFrameRate(RecordingRate, bTransactRecording);
+			Controller.SetNumberOfFrames(NumFrames, bTransactRecording);
+
+			// add to real curve data 
+			static int32 NumKeys = 3; // Muscle activation goes 0 - 1 - 0
+
+			for (int32 CurveIdx = 0; CurveIdx < CurveNamesArray.Num(); ++CurveIdx)
+			{
+				FName CurveName = CurveNamesArray[CurveIdx];
+				const FFloatCurve* FloatCurveData = nullptr;
+
+				TArray<float> TimesToRecord;
+				TArray<float> ValuesToRecord;
+				TimesToRecord.SetNum(NumKeys);
+				ValuesToRecord.SetNum(NumKeys);
+
+				bool bSeenThisCurve = false;
+				int32 WriteIndex = 0;
+				for (int32 KeyIndex = 0; KeyIndex < NumKeys; ++KeyIndex)
+				{
+					const float TimeToRecord = RecordingRate.AsSeconds(NumFramesPerMuscle * CurveIdx + NumFramesPerMuscle * KeyIndex / 2);
+					bool bIsCurveValid = false;
+					const float CurCurveValue = KeyIndex == 1 ? 1 : 0; // 0 - 1 - 0 curve value for 3 keys
+					if (!bSeenThisCurve)
+					{
+						bSeenThisCurve = true;
+						const FAnimationCurveIdentifier CurveId(CurveName, ERawCurveTrackTypes::RCT_Float);
+						Controller.AddCurve(CurveId, AACF_DefaultCurve, bTransactRecording);
+						FloatCurveData = AssetToSet->GetDataModel()->FindFloatCurve(CurveId);
+					}
+
+					if (FloatCurveData)
+					{
+						TimesToRecord[WriteIndex] = TimeToRecord;
+						ValuesToRecord[WriteIndex] = CurCurveValue;
+						++WriteIndex;
+					}
+				}
+
+				// Fill all the curve data at once
+				if (FloatCurveData)
+				{
+					TArray<FRichCurveKey> Keys;
+					for (int32 Index = 0; Index < WriteIndex; ++Index)
+					{
+						FRichCurveKey Key(TimesToRecord[Index], ValuesToRecord[Index]);
+						Key.InterpMode = ERichCurveInterpMode::RCIM_Linear;
+						Key.TangentMode = ERichCurveTangentMode::RCTM_SmartAuto;
+						Keys.Add(Key);
+					}
+
+					const FAnimationCurveIdentifier CurveId(FloatCurveData->GetName(), ERawCurveTrackTypes::RCT_Float);
+					Controller.SetCurveKeys(CurveId, Keys, bTransactRecording);
+				}
+			}
+
+			// Populate bone tracks
+			for (int32 TrackIndex = 0; TrackIndex < TrackNames.Num(); ++TrackIndex)
+			{
+				const FRawAnimSequenceTrack& RawTrack = RawTracks[TrackIndex];
+				FName BoneName = TrackNames[TrackIndex];
+				Controller.SetBoneTrackKeys(BoneName, RawTrack.PosKeys, RawTrack.RotKeys, RawTrack.ScaleKeys, bTransactRecording);
+			}
+
+			AssetToSet->PostEditChange();
+			AssetToSet->MarkPackageDirty();
+
+			// save the package to disk, for convenience and so we can run this in standalone mode
+			UPackage* const Package = AssetToSet->GetOutermost();
+			FString const PackageName = Package->GetName();
+			FString const PackageFileName = FPackageName::LongPackageNameToFilename(PackageName, FPackageName::GetAssetPackageExtension());
+
+			FSavePackageArgs SaveArgs;
+			SaveArgs.TopLevelFlags = RF_Standalone;
+			SaveArgs.SaveFlags = SAVE_NoError;
+			UPackage::SavePackage(Package, nullptr, *PackageFileName, SaveArgs);
+		} // if (AssetToSet)
+#endif // WITH_EDITOR
+	}
+}
\ No newline at end of file
diff --git a/Engine/Plugins/Experimental/ChaosFlesh/Source/ChaosFleshNodes/Private/Dataflow/ChaosFleshNodesPlugin.cpp b/Engine/Plugins/Experimental/ChaosFlesh/Source/ChaosFleshNodes/Private/Dataflow/ChaosFleshNodesPlugin.cpp
--- a/Engine/Plugins/Experimental/ChaosFlesh/Source/ChaosFleshNodes/Private/Dataflow/ChaosFleshNodesPlugin.cpp
+++ b/Engine/Plugins/Experimental/ChaosFlesh/Source/ChaosFleshNodes/Private/Dataflow/ChaosFleshNodesPlugin.cpp
@@ -50,6 +50,7 @@ void IChaosFleshNodesPlugin::StartupModule()
 	DATAFLOW_NODE_REGISTER_CREATION_FACTORY(FComputeMuscleActivationDataNode);
 	DATAFLOW_NODE_REGISTER_CREATION_FACTORY(FComputeMuscleActivationDataNode_v2);
 	DATAFLOW_NODE_REGISTER_CREATION_FACTORY(FCreateTetrahedronDataflowNode);
+	DATAFLOW_NODE_REGISTER_CREATION_FACTORY(FCurveSamplingAnimationAssetTerminalNode);
 	DATAFLOW_NODE_REGISTER_CREATION_FACTORY(FDeleteFleshVerticesDataflowNode);
 	DATAFLOW_NODE_REGISTER_CREATION_FACTORY(FDeleteVertexTrianglePositionTargetBindingDataflowNode);
 	DATAFLOW_NODE_REGISTER_CREATION_FACTORY(FFleshAssetTerminalDataflowNode);
diff --git a/Engine/Plugins/Experimental/ChaosFlesh/Source/ChaosFleshNodes/Public/Dataflow/ChaosFleshComputeMuscleActivationNode.h b/Engine/Plugins/Experimental/ChaosFlesh/Source/ChaosFleshNodes/Public/Dataflow/ChaosFleshComputeMuscleActivationNode.h
--- a/Engine/Plugins/Experimental/ChaosFlesh/Source/ChaosFleshNodes/Public/Dataflow/ChaosFleshComputeMuscleActivationNode.h
+++ b/Engine/Plugins/Experimental/ChaosFlesh/Source/ChaosFleshNodes/Public/Dataflow/ChaosFleshComputeMuscleActivationNode.h
@@ -192,6 +192,9 @@ private:
 	UPROPERTY(EditAnywhere, Category = "Dataflow", meta = (DataflowInput, DisplayName = "SkeletalMesh"))
 	TObjectPtr<const USkeletalMesh> SkeletalMesh = nullptr;
 
+	UPROPERTY(meta = (DataflowOutput))
+	FDataflowGeometrySelection GeometrySelection;
+
 	/** Click on this button to import curve names from the skeletal mesh. */
 	UPROPERTY(EditAnywhere, Category = "Skeletal Mesh Curves", meta = (ButtonImage = "Icons.Import"))
 	FDataflowFunctionProperty ImportSKMCurveNames;
diff --git a/Engine/Plugins/Experimental/ChaosFlesh/Source/ChaosFleshNodes/Public/Dataflow/ChaosFleshFleshAssetTerminalNode.h b/Engine/Plugins/Experimental/ChaosFlesh/Source/ChaosFleshNodes/Public/Dataflow/ChaosFleshFleshAssetTerminalNode.h
--- a/Engine/Plugins/Experimental/ChaosFlesh/Source/ChaosFleshNodes/Public/Dataflow/ChaosFleshFleshAssetTerminalNode.h
+++ b/Engine/Plugins/Experimental/ChaosFlesh/Source/ChaosFleshNodes/Public/Dataflow/ChaosFleshFleshAssetTerminalNode.h
@@ -9,6 +9,9 @@
 
 #include "ChaosFleshFleshAssetTerminalNode.generated.h"
 
+class UAnimSequence;
+class USkeletalMesh;
+
 USTRUCT(meta = (DataflowFlesh, DataflowTerminal))
 struct FFleshAssetTerminalDataflowNode : public FDataflowTerminalNode
 {
@@ -39,4 +42,41 @@ public:
 	virtual void Evaluate(UE::Dataflow::FContext& Context, const FDataflowOutput* Out) const override;
 };
 
+/*
+* terminal node to create an animation asset for muscle activation MLD training. 
+* The animation remains in the rest pose with the curves spiking from 0 to 1 to 0 in each block of frames (Number of frames per muscle)
+* Curves will stay at value 0 most of the time. Only one curve is active at a time.
+* Total animation frames = Frame Rate * Number of frames per muscle
+*/
+USTRUCT(meta = (DataflowFlesh, DataflowTerminal))
+struct FCurveSamplingAnimationAssetTerminalNode : public FDataflowTerminalNode
+{
+	GENERATED_USTRUCT_BODY()
+	DATAFLOW_NODE_DEFINE_INTERNAL(FCurveSamplingAnimationAssetTerminalNode, "CurveSamplingAnimationAssetTerminal", "Terminal", "")
+
+private:
+	UPROPERTY(EditAnywhere, Category = Asset, meta = (DataflowInput))
+	TObjectPtr<USkeletalMesh> SkeletalMeshAsset = nullptr;
+
+	UPROPERTY(EditAnywhere, Category = Asset, meta = (DataflowInput))
+	TObjectPtr<UAnimSequence> AnimationAsset = nullptr;
+
+	/** Frame rate of the animation */
+	UPROPERTY(EditAnywhere, Category = "Record Parameter")
+	int32 FrameRate = 30;
+
+	/** Number of frames created for each curve. Within this window of frames, curve value will go from 0 to 1 to 0. */
+	UPROPERTY(EditAnywhere, Category = "Record Parameter", DisplayName = "Number of Frames Per Muscle")
+	int32 NumFramesPerMuscle = 10;
 
+	virtual void Evaluate(UE::Dataflow::FContext& Context) const override;
+	virtual void SetAssetValue(TObjectPtr<UObject> Asset, UE::Dataflow::FContext& Context) const override;
+
+public:
+	FCurveSamplingAnimationAssetTerminalNode(const UE::Dataflow::FNodeParameters& InParam, FGuid InGuid = FGuid::NewGuid())
+	: FDataflowTerminalNode(InParam, InGuid)
+	{
+		RegisterInputConnection(&SkeletalMeshAsset);
+		RegisterInputConnection(&AnimationAsset);
+	}
+};
diff --git a/Engine/Source/Runtime/Experimental/Chaos/Private/GeometryCollection/Facades/CollectionMuscleActivationFacade.cpp b/Engine/Source/Runtime/Experimental/Chaos/Private/GeometryCollection/Facades/CollectionMuscleActivationFacade.cpp
--- a/Engine/Source/Runtime/Experimental/Chaos/Private/GeometryCollection/Facades/CollectionMuscleActivationFacade.cpp
+++ b/Engine/Source/Runtime/Experimental/Chaos/Private/GeometryCollection/Facades/CollectionMuscleActivationFacade.cpp
@@ -196,6 +196,15 @@ namespace GeometryCollection::Facades
 		return INDEX_NONE;
 	}
 
+	int32 FMuscleActivationFacade::FindMuscleGeometryIndex(const int32 MuscleIndex) const
+	{
+		if (GeometryGroupIndexAttribute.IsValidIndex(MuscleIndex))
+		{
+			return GeometryGroupIndexAttribute[MuscleIndex];
+		}
+		return INDEX_NONE;
+	}
+
 	int32 FMuscleActivationFacade::RemoveInvalidMuscles()
 	{
 		check(!IsConst());
@@ -764,15 +773,14 @@ namespace GeometryCollection::Facades
 		return MuscleLineSegments;
 	}
 
-	bool FMuscleActivationFacade::AssignCurveName(const FString& CurveName, const FString& MuscleName)
+	int32 FMuscleActivationFacade::AssignCurveName(const FString& CurveName, const FString& MuscleName)
 	{
 		const int32 MuscleIdx = FindMuscleIndexByName(MuscleName);
 		if (MuscleActivationCurveNameAttribute.IsValidIndex(MuscleIdx))
 		{
 			MuscleActivationCurveNameAttribute.Modify()[MuscleIdx] = CurveName;
-			return true;
 		}
-		return false;
+		return MuscleIdx;
 	}
 
 	TArray<int32> FMuscleActivationFacade::FindMuscleIndexByCurveName(const FString& CurveName) const
diff --git a/Engine/Source/Runtime/Experimental/Chaos/Public/GeometryCollection/Facades/CollectionMuscleActivationFacade.h b/Engine/Source/Runtime/Experimental/Chaos/Public/GeometryCollection/Facades/CollectionMuscleActivationFacade.h
--- a/Engine/Source/Runtime/Experimental/Chaos/Public/GeometryCollection/Facades/CollectionMuscleActivationFacade.h
+++ b/Engine/Source/Runtime/Experimental/Chaos/Public/GeometryCollection/Facades/CollectionMuscleActivationFacade.h
@@ -68,6 +68,7 @@ namespace GeometryCollection::Facades
 		CHAOS_API int32 NumMuscleVertices(const int32 MuscleIndex) const;
 		CHAOS_API FString FindMuscleName(const int32 MuscleIndex) const;
 		CHAOS_API int32 FindMuscleIndexByName(const FString MuscleName) const;
+		CHAOS_API int32 FindMuscleGeometryIndex(const int32 MuscleIndex) const;
 		CHAOS_API int32 RemoveInvalidMuscles();
 		CHAOS_API bool SetUpMuscleActivation(const TArray<int32>& Origin, const TArray<int32>& Insertion);
 		CHAOS_API bool SetUpMuscleActivation(const TArray<int32>& Origin, const TArray<int32>& Insertion, float ContractionVolumeScale);
@@ -82,7 +83,7 @@ namespace GeometryCollection::Facades
 			float InGlobalMuscleLengthRatioThresholdForMaxActivation);
 		CHAOS_API TArray<TArray<TArray<FVector3f>>> BuildStreamlines(const TArray<int32>& Origin, const TArray<int32>& Insertion,
 			 int32 NumLinesMultiplier, int32 MaxStreamlineIterations, int32 MaxPointsPerLine);
-		CHAOS_API bool AssignCurveName(const FString& CurveName, const FString& MuscleName);
+		CHAOS_API int32 AssignCurveName(const FString& CurveName, const FString& MuscleName);
 		CHAOS_API TArray<int32> FindMuscleIndexByCurveName(const FString& CurveName) const;
 	private:
 		const FManagedArrayCollection& ConstCollection;
