ChaosFlesh:
-Fixed GenerateSurfaceBindings node with v2: resolved render mesh interpolation noise with closest triangle query instead of smooth projection. The node renders color-coded render mesh for visual debugging.
- Dataflow RenderDynamicMesh now takes vertex color instead of only using default surface color.
- SmoothProject now checks for closest triangle

diff --git a/Engine/Plugins/Experimental/ChaosFlesh/Source/ChaosFleshNodes/ChaosFleshNodes.Build.cs b/Engine/Plugins/Experimental/ChaosFlesh/Source/ChaosFleshNodes/ChaosFleshNodes.Build.cs
--- a/Engine/Plugins/Experimental/ChaosFlesh/Source/ChaosFleshNodes/ChaosFleshNodes.Build.cs
+++ b/Engine/Plugins/Experimental/ChaosFlesh/Source/ChaosFleshNodes/ChaosFleshNodes.Build.cs
@@ -21,6 +21,7 @@ namespace UnrealBuildTool.Rules
 					"ChaosFleshEngine",
 					"Engine",
 					"GeometryAlgorithms",
+					"GeometryFramework",
 					"GeometryCore",
 					"MeshConversion",
 					"MeshConversionEngineTypes",
diff --git a/Engine/Plugins/Experimental/ChaosFlesh/Source/ChaosFleshNodes/Private/Dataflow/ChaosFleshGenerateSurfaceBindingsNode.cpp b/Engine/Plugins/Experimental/ChaosFlesh/Source/ChaosFleshNodes/Private/Dataflow/ChaosFleshGenerateSurfaceBindingsNode.cpp
--- a/Engine/Plugins/Experimental/ChaosFlesh/Source/ChaosFleshNodes/Private/Dataflow/ChaosFleshGenerateSurfaceBindingsNode.cpp
+++ b/Engine/Plugins/Experimental/ChaosFlesh/Source/ChaosFleshNodes/Private/Dataflow/ChaosFleshGenerateSurfaceBindingsNode.cpp
@@ -22,6 +22,11 @@
 #include "Rendering/SkeletalMeshRenderData.h"
 #include "UObject/PrimaryAssetId.h"
 #include "Dataflow/ChaosFleshNodesUtility.h"
+#include "UDynamicMesh.h"
+#include "DynamicMesh/InfoTypes.h"
+#include "DynamicMesh/MeshNormals.h"
+#include "Chaos/TriangleCollisionPoint.h"
+#include "Chaos/HierarchicalSpatialHash.h"
 
 DEFINE_LOG_CATEGORY(LogMeshBindings);
 
@@ -125,11 +130,43 @@ BuildVertexToVertexAdjacencyBuffer(
 	}
 }
 
+void
+BuildTriangles(
+	const FSkeletalMeshLODRenderData& LodRenderData,
+	TArray<FIntVector>& TriangleMesh)
+{
+	const FRawStaticIndexBuffer16or32Interface* IndexBuffer = LodRenderData.MultiSizeIndexContainer.GetIndexBuffer();
+	const uint32 IndexCount = IndexBuffer->Num();
 
+	TriangleMesh.Empty(); // clear, to init clean
+
+	int32 BaseTriangle = 0;
+	for (int32 SectionIndex = 0; SectionIndex < LodRenderData.RenderSections.Num(); ++SectionIndex)
+	{
+		FSkelMeshRenderSection const& RenderSection = LodRenderData.RenderSections[SectionIndex];
+		int32 NumTriangles = RenderSection.NumTriangles;
+		TriangleMesh.SetNum(BaseTriangle + NumTriangles);
+		for (int32 TriangleIt = BaseTriangle; TriangleIt < BaseTriangle + NumTriangles; ++TriangleIt)
+		{
+			TriangleMesh[TriangleIt] = FIntVector
+			(
+				IndexBuffer->Get(TriangleIt * 3 + 0),
+				IndexBuffer->Get(TriangleIt * 3 + 1),
+				IndexBuffer->Get(TriangleIt * 3 + 2)
+			);
+		}
+		BaseTriangle += NumTriangles;
+	}
+}
 
 void
 FGenerateSurfaceBindings::Evaluate(UE::Dataflow::FContext& Context, const FDataflowOutput* Out) const
 {
+	using namespace UE::Geometry;
+	auto FloatVert = [](FVector3d V) { return FVector3f(V.X, V.Y, V.Z); };
+	auto DoubleVert = [](FVector3f V) { return FVector3d(V.X, V.Y, V.Z); };
+	auto TVec3Vert = [](FVector3f V) { return Chaos::TVec3<Chaos::FRealDouble>(V.X, V.Y, V.Z); };
+	
 	if (Out->IsA<FManagedArrayCollection>(&Collection))
 	{
 		TUniquePtr<FTetrahedralCollection> InCollection(GetValue<FManagedArrayCollection>(Context, &Collection).NewCopy<FTetrahedralCollection>());
@@ -162,6 +199,9 @@ FGenerateSurfaceBindings::Evaluate(UE::Dataflow::FContext& Context, const FDataf
 
 		TObjectPtr<const USkeletalMesh> SkeletalMesh = GetValue<TObjectPtr<const USkeletalMesh>>(Context, &SkeletalMeshIn);
 		TObjectPtr<const UStaticMesh> StaticMesh = GetValue<TObjectPtr<const UStaticMesh>>(Context, &StaticMeshIn);
+		TObjectPtr<UDynamicMesh> OutSKMDynamicMesh = NewObject<UDynamicMesh>();
+		OutSKMDynamicMesh->Reset();
+		FDynamicMesh3& OutSKMDynamicMesh3 = OutSKMDynamicMesh->GetMeshRef();
 		if ((SkeletalMesh || StaticMesh) && 
 			Tetrahedron && TetrahedronStart && TetrahedronCount &&
 			Triangle && FacesStart && FacesCount &&
@@ -170,6 +210,7 @@ FGenerateSurfaceBindings::Evaluate(UE::Dataflow::FContext& Context, const FDataf
 			// Extract positions to bind
 			FString MeshId;
 			TArray<TArray<FVector3f>> MeshVertices;
+			TArray<TArray<FIntVector>> SKMMeshTriangles;
 			TArray<TArray<TArray<uint32>>> MeshNeighborNodes;
 			if (SkeletalMesh)
 			{
@@ -181,6 +222,7 @@ FGenerateSurfaceBindings::Evaluate(UE::Dataflow::FContext& Context, const FDataf
 					FSkeletalMeshRenderData* RenderData = SkeletalMesh->GetResourceForRendering();
 
 					MeshVertices.SetNum(RenderData->LODRenderData.Num());
+					SKMMeshTriangles.SetNum(RenderData->LODRenderData.Num());
 					MeshNeighborNodes.SetNum(RenderData->LODRenderData.Num());
 					for (int32 i = 0; i < RenderData->LODRenderData.Num(); i++)
 					{
@@ -198,6 +240,7 @@ FGenerateSurfaceBindings::Evaluate(UE::Dataflow::FContext& Context, const FDataf
 
 						TArray<TArray<uint32>>& NeighborNodes = MeshNeighborNodes[i];
 						BuildVertexToVertexAdjacencyBuffer(*LODRenderData, NeighborNodes);
+						BuildTriangles(*LODRenderData, SKMMeshTriangles[i]);
 					}
 				}
 #if WITH_EDITOR
@@ -205,6 +248,7 @@ FGenerateSurfaceBindings::Evaluate(UE::Dataflow::FContext& Context, const FDataf
 				{
 					const int32 LODIndex = 0;
 					MeshVertices.SetNum(1);
+					SKMMeshTriangles.SetNum(1);
 					MeshNeighborNodes.SetNum(1);
 
 					// Check first if we have bulk data available and non-empty.
@@ -212,11 +256,6 @@ FGenerateSurfaceBindings::Evaluate(UE::Dataflow::FContext& Context, const FDataf
 #if WITH_EDITORONLY_DATA
 					if (SkeletalMesh->HasMeshDescription(LODIndex))
 					{
-						// @todo(brice) : Confirm correct in release 5.4
-						//FSkeletalMeshImportData SkeletalMeshImportData;
-						//SkeletalMesh->LoadLODImportedData(LODIndex, SkeletalMeshImportData);
-						//FSkeletalMeshBuildSettings Settings;
-						//SkeletalMeshImportData.GetMeshDescription(SkeletalMesh, &Settings, SourceMesh);
 						SkeletalMesh->CloneMeshDescription(LODIndex, SourceMesh);
 					}
 					else
@@ -232,11 +271,10 @@ FGenerateSurfaceBindings::Evaluate(UE::Dataflow::FContext& Context, const FDataf
 						}
 					}
 
-					UE::Geometry::FDynamicMesh3 DynamicMesh;
+					FDynamicMesh3 DynamicMesh;
 					FMeshDescriptionToDynamicMesh Converter;
 					Converter.Convert(&SourceMesh, DynamicMesh);
 
-					//const FVertexArray& SourceVertices = SourceMesh.Vertices();
 					TArray<FVector3f>& Vertices = MeshVertices[LODIndex];
 					Vertices.SetNumUninitialized(DynamicMesh.VertexCount());
 					for (int32 j = 0; j < DynamicMesh.VertexCount(); j++)
@@ -257,6 +295,7 @@ FGenerateSurfaceBindings::Evaluate(UE::Dataflow::FContext& Context, const FDataf
 				const FStaticMeshRenderData* RenderData = StaticMesh->GetRenderData();
 				const int32 NumLOD = RenderData->LODResources.Num();
 				MeshVertices.SetNum(NumLOD);
+				SKMMeshTriangles.SetNum(NumLOD);
 				for (int32 i = 0; i < NumLOD; i++)
 				{
 					const FStaticMeshLODResources& LODResources = RenderData->LODResources[i];
@@ -281,17 +320,31 @@ FGenerateSurfaceBindings::Evaluate(UE::Dataflow::FContext& Context, const FDataf
 			//
 			// Init boundary mesh for projections.
 			//
+			TArray<FIntVector> Triangles;
 			Chaos::FTriangleMesh TetBoundaryMesh;
 			Chaos::FTriangleMesh::TBVHType<Chaos::FRealDouble> TetBoundaryBVH;
 			TArray<Chaos::FVec3> VertexD;
 			TConstArrayView<Chaos::TVec3<Chaos::FRealDouble>> VertexDView(VertexD);
 			TArray<Chaos::FVec3> PointNormals;
 
+			Chaos::THierarchicalSpatialHash<int32, Chaos::FRealDouble> SpatialHash;
+
 			TArray<int32> GeometryGroupSelected;
 			if (IsConnected(&TransformSelection))
 			{
 				FDataflowTransformSelection InTransformSelection = GetValue<FDataflowTransformSelection>(Context, &TransformSelection);
-				GeometryGroupSelected = InCollection->TransformSelectionToGeometryIndices(InTransformSelection.AsArray());
+				if (InTransformSelection.Num() == InCollection->NumElements(FGeometryCollection::TransformGroup))
+				{
+					GeometryGroupSelected = InCollection->TransformSelectionToGeometryIndices(InTransformSelection.AsArray());
+				}
+				else
+				{
+					UE_LOG(LogMeshBindings, Error,
+						TEXT("TransformSelection size: %d is different from Collection Transform group size: %d"),
+						InTransformSelection.Num(),
+						InCollection->NumElements(FGeometryCollection::TransformGroup));
+					return;
+				}
 			}
 			else
 			{
@@ -303,14 +356,19 @@ FGenerateSurfaceBindings::Evaluate(UE::Dataflow::FContext& Context, const FDataf
 			TArray<int32> UsedTetsIndexToGlobalTetIndex;
 			TArray<int32> GlobalTetIndexToUsedTetsIndex;
 			GlobalTetIndexToUsedTetsIndex.Init(INDEX_NONE, Tetrahedron->Num());
+			int32 NumGuidHit = 0;
 			for (const int32& TetMeshIdx: GeometryGroupSelected)
 			{
 				if (GeometryGroupGuidsLocal.Num() && Guids)
 				{
-					if (!GeometryGroupGuidsLocal.Contains((*Guids)[TetMeshIdx]))
+					if (!(*Guids)[TetMeshIdx].IsEmpty() && !GeometryGroupGuidsLocal.Contains((*Guids)[TetMeshIdx]))
 					{
 						continue;
 					}
+					else
+					{
+						NumGuidHit++;
+					}
 				}
 				const int32 TetMeshStart = (*TetrahedronStart)[TetMeshIdx];
 				const int32 TetMeshCount = (*TetrahedronCount)[TetMeshIdx];
@@ -331,22 +389,20 @@ FGenerateSurfaceBindings::Evaluate(UE::Dataflow::FContext& Context, const FDataf
 
 					if (bDoSurfaceProjection)
 					{
-						TetBoundaryMesh.Init(
-								reinterpret_cast<const TArray<Chaos::TVec3<int32>>&>(Triangle->GetConstArray()), 0, -1, false);
-						
-						// Promote vertices to double because that's what FTriangleMesh wants.
-						VertexD.SetNumUninitialized(Vertex->Num());
-						for (int32 i = 0; i < VertexD.Num(); i++)
+						for (int32 FaceIdx = (*FacesStart)[TetMeshIdx]; FaceIdx < (*FacesStart)[TetMeshIdx] + (*FacesCount)[TetMeshIdx]; ++FaceIdx)
 						{
-							VertexD[i] = Chaos::FVec3((*Vertex)[i][0], (*Vertex)[i][1], (*Vertex)[i][2]);
+							Triangles.Add((*Triangle)[FaceIdx]);
 						}
-						VertexDView = TConstArrayView<Chaos::TVec3<Chaos::FRealDouble>>(VertexD);
-
-						PointNormals = TetBoundaryMesh.GetPointNormals(VertexDView, false, true);
-						TetBoundaryMesh.BuildBVH(VertexDView, TetBoundaryBVH);
 					}
 				}
 			}
+			if (GeometryGroupGuidsLocal.Num() && NumGuidHit == 0)
+			{
+				UE_LOG(LogMeshBindings, Error,
+					TEXT("GeometryGroupGuids contains %d guids but none was matched (empty guids are ignored)."),
+					GeometryGroupGuidsLocal.Num());
+				return;
+			}
 			// Init BVH for tetrahedra.
 			BVHTetPtrs.SetNumUninitialized(Tets.Num());
 			for (int32 TetIdx = 0; TetIdx < Tets.Num(); ++TetIdx)
@@ -359,6 +415,28 @@ FGenerateSurfaceBindings::Evaluate(UE::Dataflow::FContext& Context, const FDataf
 				Chaos::FReal,
 				3> TetBVH(BVHTetPtrs);
 
+			// Init BVH for surface triangle mesh.
+			if (bDoSurfaceProjection)
+			{
+				//Todo(Chaosflesh): refactor reinterpret_cast in case the memory layout of the vector types not the same
+				TetBoundaryMesh.Init(
+					reinterpret_cast<const TArray<Chaos::TVec3<int32>>&>(Triangles), 0, -1, false);
+
+				// Promote vertices to double because that's what FTriangleMesh wants.
+				VertexD.SetNumUninitialized(Vertex->Num());
+				for (int32 i = 0; i < VertexD.Num(); i++)
+				{
+					VertexD[i] = Chaos::FVec3((*Vertex)[i][0], (*Vertex)[i][1], (*Vertex)[i][2]);
+				}
+				VertexDView = TConstArrayView<Chaos::TVec3<Chaos::FRealDouble>>(VertexD);
+
+				PointNormals = TetBoundaryMesh.GetPointNormals(VertexDView, false, true);
+				TetBoundaryMesh.BuildBVH(VertexDView, TetBoundaryBVH);
+
+
+				TetBoundaryMesh.BuildSpatialHash(VertexDView, SpatialHash, Chaos::FRealDouble(SurfaceProjectionSearchRadius));
+			}
+
 			//
 			// Do intersection tests against tets, then the surface.
 			//
@@ -368,6 +446,7 @@ FGenerateSurfaceBindings::Evaluate(UE::Dataflow::FContext& Context, const FDataf
 			TArray<TArray<FVector3f>> Offsets;	 Offsets.SetNum(MeshVertices.Num());
 			TArray<TArray<float>> Masks;		 Masks.SetNum(MeshVertices.Num());
 			TArray<int32> Orphans;
+			TArray<int32> OrphansCopy;
 			int32 TetHits = 0;
 			int32 TriHits = 0;
 			int32 Adoptions = 0;
@@ -378,6 +457,7 @@ FGenerateSurfaceBindings::Evaluate(UE::Dataflow::FContext& Context, const FDataf
 				Weights[LOD].SetNumUninitialized(MeshVertices[LOD].Num());
 				Offsets[LOD].SetNumUninitialized(MeshVertices[LOD].Num());
 				Masks[LOD].SetNumUninitialized(MeshVertices[LOD].Num());
+				int32 NumTetNotCollocated = 0;
 
 				TArray<int32> TetIntersections; TetIntersections.Reserve(64);
 				for (int32 i = 0; i < MeshVertices[LOD].Num(); i++)
@@ -414,46 +494,60 @@ FGenerateSurfaceBindings::Evaluate(UE::Dataflow::FContext& Context, const FDataf
 							}
 							else
 							{
-								UE_LOG(LogMeshBindings, Warning,
-									TEXT("Vertex position does not collocate with interpolated position in tetrahedron %d, distance = %.4f)"),
-									TetIdx,
-									(Pos - EmbeddedPos).Length());
+								NumTetNotCollocated++;
+								if (NumTetNotCollocated == 1)
+								{
+									UE_LOG(LogMeshBindings, Error,
+										TEXT("Vertex position does not collocate with interpolated position, for example LOD %d, SKM vertex %d, tetrahedron %d, distance = %.4f)"),
+										LOD,
+										i,
+										TetIdx,
+										(Pos - EmbeddedPos).Length());
+								}
 							}
 						}
 					}
 					if (j == TetIntersections.Num())
 					{
-						// This vertex didn't land inside any tetrahedra. Project it to the tet boundary surface.
-						int32 TriIdx = INDEX_NONE;
-						Chaos::FVec3 TriWeights;
-						if (bDoSurfaceProjection &&
-							TetBoundaryMesh.SmoothProject(
-								TetBoundaryBVH,
-								VertexDView,
-								PointNormals, Pos,
-								TriIdx, TriWeights, SurfaceProjectionIterations))
+						bool bSuccess = false;
+						if (bDoSurfaceProjection)
 						{
-							TriHits++;
-							const FIntVector& Tri = (*Triangle)[TriIdx];
-							Parents[LOD][i][0] = Tri[0];
-							Parents[LOD][i][1] = Tri[1];
-							Parents[LOD][i][2] = Tri[2];
-							Parents[LOD][i][3] = INDEX_NONE;
-
-							Weights[LOD][i][0] = TriWeights[0];
-							Weights[LOD][i][1] = TriWeights[1];
-							Weights[LOD][i][2] = TriWeights[2];
-							Weights[LOD][i][3] = 0.0;
-
-							const FVector3f EmbeddedPos =
-								TriWeights[0] * Vertex->GetConstArray()[Tri[0]] +
-								TriWeights[1] * Vertex->GetConstArray()[Tri[1]] +
-								TriWeights[2] * Vertex->GetConstArray()[Tri[2]];
-							Offsets[LOD][i] = EmbeddedPos - Pos;
-
-							Masks[LOD][i] = 1.0; // Shader does sim for this vertex
+							TArray<Chaos::TTriangleCollisionPoint<Chaos::FRealDouble>> Result;
+							//PointClosestTriangleQuery instead of SmoothProject
+							if (TetBoundaryMesh.PointClosestTriangleQuery(SpatialHash, VertexDView, i, TVec3Vert(Pos), Chaos::FRealDouble(SurfaceProjectionSearchRadius), Chaos::FRealDouble(SurfaceProjectionSearchRadius),
+								[](const int32 PointIndex, const int32 TriangleIndex)->bool
+								{
+									// use all nearby triangles
+									return true;
+								},
+								Result))
+							{
+								for (const Chaos::TTriangleCollisionPoint<Chaos::FRealDouble>& CollisionPoint : Result)
+								{
+									const FIntVector& Tri = Triangles[CollisionPoint.Indices[1]];
+									TriHits++;
+									Parents[LOD][i][0] = Tri[0];
+									Parents[LOD][i][1] = Tri[1];
+									Parents[LOD][i][2] = Tri[2];
+									Parents[LOD][i][3] = INDEX_NONE;
+
+									Weights[LOD][i][0] = CollisionPoint.Bary[1];
+									Weights[LOD][i][1] = CollisionPoint.Bary[2];
+									Weights[LOD][i][2] = CollisionPoint.Bary[3];
+									Weights[LOD][i][3] = 0.0;
+
+									const FVector3f EmbeddedPos =
+										Weights[LOD][i][0] * Vertex->GetConstArray()[Tri[0]] +
+										Weights[LOD][i][1] * Vertex->GetConstArray()[Tri[1]] +
+										Weights[LOD][i][2] * Vertex->GetConstArray()[Tri[2]];
+									Offsets[LOD][i] = EmbeddedPos - Pos;
+									Masks[LOD][i] = 1.0; // Shader does sim for this vertex
+									bSuccess = true;
+									break;
+								}
+							}
 						}
-						else
+						if (!bSuccess)
 						{
 							// Despair...
 							Orphans.Add(i);
@@ -473,7 +567,7 @@ FGenerateSurfaceBindings::Evaluate(UE::Dataflow::FContext& Context, const FDataf
 							Offsets[LOD][i][2] = 0.0;
 
 							Masks[LOD][i] = 0.0; // Shader does skinning for this vertex
-						} // if !SmoothProject()
+						}
 					} // if !TetIntersections
 				} // end for all vertices
 
@@ -485,6 +579,7 @@ FGenerateSurfaceBindings::Evaluate(UE::Dataflow::FContext& Context, const FDataf
 					continue;
 				}
 				const TArray<TArray<uint32>>& NeighborNodes = MeshNeighborNodes[LOD];
+				OrphansCopy = Orphans;
 				TSet<int32> OrphanSet(Orphans);
 				while (bDoOrphanReparenting && Orphans.Num())
 				{
@@ -570,10 +665,10 @@ FGenerateSurfaceBindings::Evaluate(UE::Dataflow::FContext& Context, const FDataf
 											const Chaos::FTetrahedron& Tet = Tets[UsedTetIdx];
 
 											Chaos::TVec4<Chaos::FReal> W;
-											Chaos::TVec3<Chaos::FReal> EmbeddedPos = Tet.FindClosestPointAndBary(PosD, W, 1.0e-4);
+											Chaos::TVec3<Chaos::FReal> EmbeddedPos = Tet.FindClosestPointAndBary(PosD, W, 0 /*Tolerance*/); //Tolerance should be small negative number or zero
 											Chaos::TVec3<Chaos::FReal> O = EmbeddedPos - PosD;
 											Chaos::FReal Dist = O.SquaredLength();
-											if (Dist < CurrDist)
+											if (Dist < CurrDist) // Closest neighbor tet
 											{
 												CurrDist = Dist;
 												Parents[LOD][Orphan] = (*Tetrahedron)[TetIdx];
@@ -615,6 +710,13 @@ FGenerateSurfaceBindings::Evaluate(UE::Dataflow::FContext& Context, const FDataf
 					bUseSkeletalMeshImportModel ? TEXT("import") : TEXT("render"),
 					*MeshId, LOD,
 					MeshVertices[LOD].Num(), TetHits, TriHits, Adoptions, NumOrphans);
+				if (NumTetNotCollocated)
+				{
+					UE_LOG(LogMeshBindings, Error,
+						TEXT("%d vertex positions do not collocate with interpolated position for LOD %d"),
+						NumTetNotCollocated,
+						LOD);
+				}
 
 			} // end for all LOD
 
@@ -627,7 +729,64 @@ FGenerateSurfaceBindings::Evaluate(UE::Dataflow::FContext& Context, const FDataf
 				TetBindings.AddBindingsGroup(/*TetMeshIdx = */ 0, MeshName, LOD);
 				TetBindings.SetBindingsData(Parents[LOD], Weights[LOD], Offsets[LOD], Masks[LOD]);
 			}
+
+			//Write DynamicMesh			
+			if (MeshVertices.Num())
+			{
+				OutSKMDynamicMesh3.EnableAttributes();
+				OutSKMDynamicMesh3.EnableVertexColors(FVector3f(1, 0, 0));
+				TBitArray<> WasOrphan(false, MeshVertices[0].Num());
+				UE_LOG(LogTemp, Warning, TEXT("Num of previous orphans: %d"), OrphansCopy.Num());
+				for (int32 OrphanIdx : OrphansCopy)
+				{
+					WasOrphan[OrphanIdx] = true;
+				}
+				for (int32 VertexIndex = 0; VertexIndex < MeshVertices[0].Num(); ++VertexIndex)
+				{
+					FVertexInfo VertexInfo;
+					VertexInfo.Position = DoubleVert(MeshVertices[0][VertexIndex]);
+					VertexInfo.bHaveC = true;
+					if (Parents[0][VertexIndex][0] == INDEX_NONE)
+					{
+						VertexInfo.Color = FVector3f(1, 0, 0); //red if orphan
+					}
+					else
+					{
+						if (Parents[0][VertexIndex][3] == INDEX_NONE)
+						{
+							VertexInfo.Color = FVector3f(0, 0, 1); //blue if on surface
+						}
+						else
+						{
+							VertexInfo.Color = FVector3f(0, 1, 0); //green if in tet
+						}
+						if (WasOrphan[VertexIndex])
+						{
+							VertexInfo.Color += FVector3f(1, 0, 0); //add red if was orphan
+						}
+					}
+					OutSKMDynamicMesh3.AppendVertex(VertexInfo);
+				}
+				
+				for (int32 TriangleIndex = 0; TriangleIndex < SKMMeshTriangles[0].Num(); ++TriangleIndex)
+				{
+					OutSKMDynamicMesh3.AppendTriangle(FIndex3i(
+						SKMMeshTriangles[0][TriangleIndex][0],
+						SKMMeshTriangles[0][TriangleIndex][1],
+						SKMMeshTriangles[0][TriangleIndex][2])
+					);
+				}
+				// Compute normals
+				OutSKMDynamicMesh3.EnableVertexNormals(FVector3f(1, 0, 0));
+				FMeshNormals MeshNormals(&OutSKMDynamicMesh3);
+				MeshNormals.ComputeVertexNormals();
+				for (int32 VertexIndex = 0; VertexIndex < MeshVertices[0].Num(); ++VertexIndex)
+				{
+					OutSKMDynamicMesh3.SetVertexNormal(VertexIndex, FloatVert(MeshNormals[VertexIndex]));
+				}
+			}
 		}
 		SetValue<const FManagedArrayCollection&>(Context, *InCollection, &Collection);
+		SetValue(Context, OutSKMDynamicMesh, &SKMDynamicMesh);
 	}
-}
+}
\ No newline at end of file
diff --git a/Engine/Plugins/Experimental/ChaosFlesh/Source/ChaosFleshNodes/Public/Dataflow/ChaosFleshGenerateSurfaceBindingsNode.h b/Engine/Plugins/Experimental/ChaosFlesh/Source/ChaosFleshNodes/Public/Dataflow/ChaosFleshGenerateSurfaceBindingsNode.h
--- a/Engine/Plugins/Experimental/ChaosFlesh/Source/ChaosFleshNodes/Public/Dataflow/ChaosFleshGenerateSurfaceBindingsNode.h
+++ b/Engine/Plugins/Experimental/ChaosFlesh/Source/ChaosFleshNodes/Public/Dataflow/ChaosFleshGenerateSurfaceBindingsNode.h
@@ -12,18 +12,40 @@
 
 class UStaticMesh;
 class USkeletalMesh;
+class UDynamicMesh;
 
 DECLARE_LOG_CATEGORY_EXTERN(LogMeshBindings, Verbose, All);
 
-/** Generate barycentric bindings (used by the FleshDeformer deformer graph) of a render surface to a tetrahedral mesh. */
+/* Generate barycentric bindings (used by the FleshDeformer deformer graph and Geometry Cache generation) of a render surface to a tetrahedral mesh.
+* If a point is outside of the tetrahedral mesh, find surface embedding within SurfaceProjectionSearchRadius.
+Embeddings are color coded in the render view:
+	green: embedded on in a tetrahedron
+	blue: embedded on a surface triangle
+	red: orphan (cannot be embedded)
+	yellow: orphan reparented to a tetrahedron from a node neighbor */
 USTRUCT(meta = (DataflowFlesh))
 struct FGenerateSurfaceBindings : public FDataflowNode
 {
 	GENERATED_USTRUCT_BODY()
 	DATAFLOW_NODE_DEFINE_INTERNAL(FGenerateSurfaceBindings, "GenerateSurfaceBindings", "Flesh", "")
-	DATAFLOW_NODE_RENDER_TYPE("SurfaceRender", FGeometryCollection::StaticType(), "Collection")
+	DATAFLOW_NODE_RENDER_TYPE("SurfaceRender", FName("FDynamicMesh3"), "SKMDynamicMesh")
 
 public:
+	FGenerateSurfaceBindings(const UE::Dataflow::FNodeParameters& InParam, FGuid InGuid = FGuid::NewGuid())
+		: FDataflowNode(InParam, InGuid)
+	{
+		RegisterInputConnection(&Collection);
+		RegisterOutputConnection(&Collection, &Collection);
+		RegisterInputConnection(&StaticMeshIn);
+		RegisterInputConnection(&SkeletalMeshIn);
+		RegisterInputConnection(&TransformSelection);
+		RegisterInputConnection(&GeometryGroupGuidsIn);
+		RegisterOutputConnection(&SKMDynamicMesh)
+			.SetCanHidePin(true)
+			.SetPinIsHidden(true);
+	}
+
+private:
 	/** Passthrough geometry collection. Bindings are stored as standalone groups in the \p Collection, keyed by the name of the input render mesh and all available LOD's. */
 	UPROPERTY(meta = (DataflowInput, DataflowOutput, DataflowPassthrough = "Collection", DisplayName = "Collection"))
 	FManagedArrayCollection Collection;
@@ -52,24 +74,17 @@ public:
 	UPROPERTY(EditAnywhere, Category = "Dataflow", meta = (DisplayName = "DoSurfaceProjection"))
 	bool bDoSurfaceProjection = true;
 
-	/** The maximum number of iterations to try expanding the domain while looking for surface triangles to bind to. */
-	UPROPERTY(EditAnywhere, Category = "Dataflow", meta = (DisplayName = "SurfaceProjectionIterations", EditCondition = "bDoSurfaceProjection == true"))
-	uint32 SurfaceProjectionIterations = 10;
+	/** The search radius when looking for surface triangles to bind to. */
+	UPROPERTY(EditAnywhere, Category = "Dataflow", meta = (DisplayName = "SurfaceProjectionSearchRadius", EditCondition = "bDoSurfaceProjection == true"))
+	float SurfaceProjectionSearchRadius = 1.f;
 
 	/** When nodes aren't contained in tetrahedra and surface projection fails, try to find suitable bindings by looking to neighboring parents. */
 	UPROPERTY(EditAnywhere, Category = "Dataflow", meta = (DisplayName = "DoOrphanReparenting"))
 	bool bDoOrphanReparenting = true;
 
-	FGenerateSurfaceBindings(const UE::Dataflow::FNodeParameters& InParam, FGuid InGuid = FGuid::NewGuid())
-		: FDataflowNode(InParam, InGuid)
-	{
-		RegisterInputConnection(&Collection);
-		RegisterOutputConnection(&Collection, &Collection);
-		RegisterInputConnection(&StaticMeshIn);
-		RegisterInputConnection(&SkeletalMeshIn);
-		RegisterInputConnection(&TransformSelection);
-		RegisterInputConnection(&GeometryGroupGuidsIn);
-	}
+	/** Converted from embedded skeletal/static mesh */
+	UPROPERTY(meta = (DataflowOutput))
+	TObjectPtr<UDynamicMesh> SKMDynamicMesh;
 
 	virtual void Evaluate(UE::Dataflow::FContext& Context, const FDataflowOutput* Out) const override;
 };
\ No newline at end of file
diff --git a/Engine/Plugins/Experimental/Dataflow/Source/DataflowEditor/Private/Dataflow/DataflowEngineRendering.cpp b/Engine/Plugins/Experimental/Dataflow/Source/DataflowEditor/Private/Dataflow/DataflowEngineRendering.cpp
--- a/Engine/Plugins/Experimental/Dataflow/Source/DataflowEditor/Private/Dataflow/DataflowEngineRendering.cpp
+++ b/Engine/Plugins/Experimental/Dataflow/Source/DataflowEditor/Private/Dataflow/DataflowEngineRendering.cpp
@@ -105,11 +105,21 @@ namespace UE::Dataflow
 
 				// Add VertexNormal and VertexColor
 				TArray<FLinearColor> VertexColors; VertexColors.AddUninitialized(Vertices.Num());
-				for (int32 VertexIdx = 0; VertexIdx < VertexNormals.Num(); ++VertexIdx)
+				
+				if (DynamicMesh.HasVertexColors())
 				{
-					VertexColors[VertexIdx] = FLinearColor(FDataflowEditorModule::SurfaceColor);
+					for (int32 VertexID : DynamicMesh.VertexIndicesItr())
+					{
+						VertexColors[VertexID] = DynamicMesh.GetVertexColor(VertexID);
+					}
+				}
+				else
+				{
+					for (int32 VertexIdx = 0; VertexIdx < VertexNormals.Num(); ++VertexIdx)
+					{
+						VertexColors[VertexIdx] = FLinearColor(FDataflowEditorModule::SurfaceColor);
+					}
 				}
-
 				const int32 NumUVLayers = DynamicMesh.Attributes()->NumUVLayers();
 
 				const int32 GeometryIndex = RenderCollection.StartGeometryGroup(GeometryGroupName);
@@ -984,40 +994,45 @@ namespace UE::Dataflow
 			if (State.GetRenderOutputs().Num())
 			{
 				FName PrimaryOutput = State.GetRenderOutputs()[0]; // "VectorField"
-				if (PrimaryOutput.IsEqual(FName("VectorField")))
+				FFieldCollection Default;
+				const FFieldCollection& Collection = State.GetValue<FFieldCollection>(PrimaryOutput, Default);
+				TArray<TPair<FVector3f, FVector3f>> VectorField = Collection.GetVectorField();
+				TArray<FLinearColor> VertexColors = Collection.GetVectorColor();
+				const int32 NumVertices = 3 * VectorField.Num();
+				const int32 NumTriangles = VectorField.Num();
+
+				TArray<FVector3f> Vertices; Vertices.AddUninitialized(NumVertices);
+				TArray<FIntVector> Tris; Tris.AddUninitialized(NumTriangles);
+				TArray<FVector3f> VertexNormals; VertexNormals.AddUninitialized(NumVertices);
+
+				for (int32 i = 0; i < VectorField.Num(); i++)
 				{
-						FFieldCollection Default;
-						const FFieldCollection& Collection = State.GetValue<FFieldCollection>(PrimaryOutput, Default);
-						TArray<TPair<FVector3f, FVector3f>> VectorField = Collection.GetVectorField();
-						TArray<FLinearColor> VertexColors = Collection.GetVectorColor();
-						const int32 NumVertices = 3 * VectorField.Num();
-						const int32 NumTriangles = VectorField.Num();
-
-						TArray<FVector3f> Vertices; Vertices.AddUninitialized(NumVertices);
-						TArray<FIntVector> Tris; Tris.AddUninitialized(NumTriangles);
-						TArray<FVector3f> VertexNormals; VertexNormals.AddUninitialized(NumVertices);
-
-						for (int32 i = 0; i < VectorField.Num(); i++)
-						{
 							
-							FVector3f Dir = VectorField[i].Value - VectorField[i].Key;
-							FVector3f DirAdd = Dir;
-							DirAdd.X += 1.f;
-							FVector3f OrthogonalDir = (Dir^ DirAdd).GetSafeNormal();
-							Tris[i] = FIntVector(3*i, 3*i+1, 3*i+2);
-							Vertices[3*i] = VectorField[i].Key;
-							Vertices[3*i+1] = VectorField[i].Value;
-							Vertices[3*i+2] = VectorField[i].Key + float(0.1) * Dir.Size() * OrthogonalDir;
-							FVector3f TriangleNormal = (OrthogonalDir ^ Dir).GetSafeNormal();
-							VertexNormals[3*i] = TriangleNormal;
-							VertexNormals[3*i+1] = TriangleNormal;
-							VertexNormals[3*i+2] = TriangleNormal;
-						}
-
-					int32 GeometryIndex = RenderCollection.StartGeometryGroup(State.GetGuid().ToString());
-					RenderCollection.AddSurface(MoveTemp(Vertices), MoveTemp(Tris), MoveTemp(VertexNormals), MoveTemp(VertexColors));
-					RenderCollection.EndGeometryGroup(GeometryIndex);
+					FVector3f Dir = VectorField[i].Value - VectorField[i].Key;
+					FVector3f OrthogonalDir;
+					if (Dir[1] < UE_SMALL_NUMBER && Dir[2] < UE_SMALL_NUMBER && Dir[0] > UE_SMALL_NUMBER)
+					{
+						OrthogonalDir = FVector3f(0, 0, 1);
+					}
+					else
+					{
+						FVector3f DirAdd = Dir;
+						DirAdd.X += 1.f;
+						OrthogonalDir = (Dir ^ DirAdd).GetSafeNormal();
+					}
+					Tris[i] = FIntVector(3*i, 3*i+1, 3*i+2);
+					Vertices[3*i] = VectorField[i].Key;
+					Vertices[3*i+1] = VectorField[i].Value;
+					Vertices[3*i+2] = VectorField[i].Key + float(0.1) * Dir.Size() * OrthogonalDir;
+					FVector3f TriangleNormal = (OrthogonalDir ^ Dir).GetSafeNormal();
+					VertexNormals[3*i] = TriangleNormal;
+					VertexNormals[3*i+1] = TriangleNormal;
+					VertexNormals[3*i+2] = TriangleNormal;
 				}
+
+				int32 GeometryIndex = RenderCollection.StartGeometryGroup(State.GetGuid().ToString());
+				RenderCollection.AddSurface(MoveTemp(Vertices), MoveTemp(Tris), MoveTemp(VertexNormals), MoveTemp(VertexColors));
+				RenderCollection.EndGeometryGroup(GeometryIndex);
 			}
 		}
 	};
diff --git a/Engine/Source/Runtime/Experimental/Chaos/Private/Chaos/TriangleMesh.cpp b/Engine/Source/Runtime/Experimental/Chaos/Private/Chaos/TriangleMesh.cpp
--- a/Engine/Source/Runtime/Experimental/Chaos/Private/Chaos/TriangleMesh.cpp
+++ b/Engine/Source/Runtime/Experimental/Chaos/Private/Chaos/TriangleMesh.cpp
@@ -1753,16 +1753,30 @@ bool FTriangleMesh::SmoothProject(
 		}
 		// Test candidates
 		TArray<FVec3> CandidateWeights;
-		TArray<bool> Success = Chaos::SmoothProject<T>(Points, MElements, PointNormals, Pos, CandidateTris, CandidateWeights, true);
+		TArray<bool> Success = Chaos::SmoothProject<T>(Points, MElements, PointNormals, Pos, CandidateTris, CandidateWeights, /*UseFirstFound = */ false);
+		double MinDistanceSq = FLT_MAX; // find closest triangle
 		for (int32 i = 0; i < Success.Num(); i++)
 		{
 			if (Success[i])
 			{
-				TriangleIndex = CandidateTris[i];
-				Weights = CandidateWeights[i];
-				return true;
+				FVec3 InterpPos = FVec3(0);
+				for (int32 LocalIdx = 0; LocalIdx < 3; ++LocalIdx)
+				{
+					InterpPos += Points[MElements[CandidateTris[i]][LocalIdx]] * CandidateWeights[i][LocalIdx];
+				}
+				double DistSq = (Pos - InterpPos).SizeSquared();
+				if (DistSq < MinDistanceSq)
+				{
+					MinDistanceSq = DistSq;
+					TriangleIndex = CandidateTris[i];
+					Weights = CandidateWeights[i];
+				}
 			}
 		}
+		if (MinDistanceSq < FLT_MAX)
+		{
+			return true;
+		}
 		// No hits. Add candidates to skip list, and go searching for more.
 		SkipTris.Append(CandidateTris);
 	} while (Iter < MaxIters);
