ChaoFlesh:
-Implemented airmesh constraint to resolve volume loss from zero rest length spring constraints.
-Added CreateAirMesh node

diff --git a/Engine/Plugins/Experimental/ChaosFlesh/Source/ChaosFleshNodes/Private/Dataflow/ChaosFleshNodesPlugin.cpp b/Engine/Plugins/Experimental/ChaosFlesh/Source/ChaosFleshNodes/Private/Dataflow/ChaosFleshNodesPlugin.cpp
--- a/Engine/Plugins/Experimental/ChaosFlesh/Source/ChaosFleshNodes/Private/Dataflow/ChaosFleshNodesPlugin.cpp
+++ b/Engine/Plugins/Experimental/ChaosFlesh/Source/ChaosFleshNodes/Private/Dataflow/ChaosFleshNodesPlugin.cpp
@@ -49,6 +49,7 @@ void IChaosFleshNodesPlugin::StartupModule()
 	DATAFLOW_NODE_REGISTER_CREATION_FACTORY(FComputeIslandsNode);
 	DATAFLOW_NODE_REGISTER_CREATION_FACTORY(FComputeMuscleActivationDataNode);
 	DATAFLOW_NODE_REGISTER_CREATION_FACTORY(FComputeMuscleActivationDataNode_v2);
+	DATAFLOW_NODE_REGISTER_CREATION_FACTORY(FCreateAirTetrahedralConstraintDataflowNode);
 	DATAFLOW_NODE_REGISTER_CREATION_FACTORY(FCreateTetrahedronDataflowNode);
 	DATAFLOW_NODE_REGISTER_CREATION_FACTORY(FCurveSamplingAnimationAssetTerminalNode);
 	DATAFLOW_NODE_REGISTER_CREATION_FACTORY(FDeleteFleshVerticesDataflowNode);
diff --git a/Engine/Plugins/Experimental/ChaosFlesh/Source/ChaosFleshNodes/Private/Dataflow/ChaosFleshSetVertexTrianglePositionTargetBindingNode.cpp b/Engine/Plugins/Experimental/ChaosFlesh/Source/ChaosFleshNodes/Private/Dataflow/ChaosFleshSetVertexTrianglePositionTargetBindingNode.cpp
--- a/Engine/Plugins/Experimental/ChaosFlesh/Source/ChaosFleshNodes/Private/Dataflow/ChaosFleshSetVertexTrianglePositionTargetBindingNode.cpp
+++ b/Engine/Plugins/Experimental/ChaosFlesh/Source/ChaosFleshNodes/Private/Dataflow/ChaosFleshSetVertexTrianglePositionTargetBindingNode.cpp
@@ -11,16 +11,20 @@
 #include "GeometryCollection/Facades/CollectionMeshFacade.h"
 #include "GeometryCollection/Facades/CollectionPositionTargetFacade.h"
 #include "GeometryCollection/Facades/CollectionCollisionFacade.h"
+#include "UDynamicMesh.h"
+#include "DynamicMesh/InfoTypes.h"
+#include "DynamicMesh/MeshNormals.h"
+#include "Dataflow/ChaosFleshNodesUtility.h"
 
 #include UE_INLINE_GENERATED_CPP_BY_NAME(ChaosFleshSetVertexTrianglePositionTargetBindingNode)
 
 void FSetVertexTrianglePositionTargetBindingDataflowNode::Evaluate(UE::Dataflow::FContext& Context, const FDataflowOutput* Out) const
 {
-	if (Out->IsA<DataType>(&Collection))
+	if (Out->IsA(&Collection))
 	{
-		DataType InCollection = GetValue<DataType>(Context, &Collection);
+		FManagedArrayCollection InCollection = GetValue(Context, &Collection);
 
-		TUniquePtr<FFleshCollection> InFleshCollection(GetValue<DataType>(Context, &Collection).NewCopy<FFleshCollection>());
+		TUniquePtr<FFleshCollection> InFleshCollection(GetValue(Context, &Collection).NewCopy<FFleshCollection>());
 
 		Chaos::FFleshCollectionFacade TetCollection(*InFleshCollection);
 		if (TManagedArray<FVector3f>* Vertices = InCollection.FindAttribute<FVector3f>("Vertex", FGeometryCollection::VerticesGroup))
@@ -183,9 +187,9 @@ void FSetVertexTrianglePositionTargetBindingDataflowNode::Evaluate(UE::Dataflow:
 
 void FDeleteVertexTrianglePositionTargetBindingDataflowNode::Evaluate(UE::Dataflow::FContext& Context, const FDataflowOutput* Out) const
 {
-	if (Out->IsA<DataType>(&Collection))
+	if (Out->IsA(&Collection))
 	{
-		DataType InCollection = GetValue<DataType>(Context, &Collection);
+		FManagedArrayCollection InCollection = GetValue(Context, &Collection);
 		if (IsConnected(&VertexSelection1) && IsConnected(&VertexSelection2))
 		{
 			const FDataflowVertexSelection& InVertexSelection1 = GetValue(Context, &VertexSelection1);
@@ -219,9 +223,9 @@ void FDeleteVertexTrianglePositionTargetBindingDataflowNode::Evaluate(UE::Datafl
 
 void FSetCollidableVerticesDataflowNode::Evaluate(UE::Dataflow::FContext& Context, const FDataflowOutput* Out) const
 {
-	if (Out->IsA<DataType>(&Collection))
+	if (Out->IsA(&Collection))
 	{
-		DataType InCollection = GetValue<DataType>(Context, &Collection);
+		FManagedArrayCollection InCollection = GetValue(Context, &Collection);
 		if (IsConnected(&VertexSelection))
 		{
 			GeometryCollection::Facades::FCollisionFacade CollisionFacade(InCollection);
@@ -230,4 +234,169 @@ void FSetCollidableVerticesDataflowNode::Evaluate(UE::Dataflow::FContext& Contex
 		}
 		SetValue(Context, MoveTemp(InCollection), &Collection);
 	}
+}
+
+void FCreateAirTetrahedralConstraintDataflowNode::Evaluate(UE::Dataflow::FContext& Context, const FDataflowOutput* Out) const
+{
+	using namespace UE::Geometry;
+	if (Out->IsA(&Collection))
+	{
+		TUniquePtr<FFleshCollection> InFleshCollection(GetValue(Context, &Collection).NewCopy<FFleshCollection>());
+		Chaos::FFleshCollectionFacade TetCollection(*InFleshCollection);
+		TObjectPtr<UDynamicMesh> OutDynamicMesh = NewObject<UDynamicMesh>();
+		if (const TManagedArray<FVector3f>* Vertices = InFleshCollection->FindAttribute<FVector3f>("Vertex", FGeometryCollection::VerticesGroup))
+		{
+			if (const TManagedArray<FIntVector>* Indices = InFleshCollection->FindAttribute<FIntVector>("Indices", FGeometryCollection::FacesGroup))
+			{
+				TManagedArray<FIntVector4>& Tetrahedron = InFleshCollection->AddAttribute<FIntVector4>(FTetrahedralCollection::TetrahedronAttribute, FTetrahedralCollection::TetrahedralGroup);
+				TArray<FVector3f> Vertex = Vertices->GetConstArray();
+				TetCollection.ComponentSpaceVertices(Vertex);
+				GeometryCollection::Facades::FCollectionMeshFacade MeshFacade(*InFleshCollection);
+				TArray<int32> ComponentIndex = MeshFacade.GetGeometryGroupIndexArray();
+				TArray<Chaos::TVector<int32, 3>> IndicesArray;
+				for (int32 i = 0; i < Indices->Num(); i++)
+				{
+					Chaos::TVector<int32, 3> CurrentIndices(0);
+					for (int32 j = 0; j < 3; j++)
+					{
+						CurrentIndices[j] = (*Indices)[i][j];
+					}
+					if (CurrentIndices[0] != INDEX_NONE
+						&& CurrentIndices[1] != INDEX_NONE
+						&& CurrentIndices[2] != INDEX_NONE)
+					{
+						IndicesArray.Emplace(CurrentIndices);
+					}
+				}
+
+				TArray<int32> IndicesMap;
+				if (IsConnected(&VertexSelection))
+				{
+					FDataflowVertexSelection InDataflowVertexSelection = GetValue<FDataflowVertexSelection>(Context, &VertexSelection);
+					IndicesMap = InDataflowVertexSelection.AsArray();
+				}
+				else
+				{
+					TArray<TArray<int32>> LocalIndex;
+					TArray<TArray<int32>>* LocalIndexPtr = &LocalIndex;
+					TArray<TArray<int>> GlobalIndex = Chaos::Utilities::ComputeIncidentElements(IndicesArray, LocalIndexPtr);
+					int32 ActualParticleCount = 0;
+					for (int32 l = 0; l < GlobalIndex.Num(); l++)
+					{
+						if (GlobalIndex[l].Num() > 0)
+						{
+							ActualParticleCount += 1;
+						}
+					}
+					IndicesMap.SetNum(ActualParticleCount);
+					int32 CurrentParticleIndex = 0;
+					for (int32 i = 0; i < GlobalIndex.Num(); i++)
+					{
+						if (GlobalIndex[i].Num() > 0)
+						{
+							IndicesMap[CurrentParticleIndex] = (*Indices)[GlobalIndex[i][0]][LocalIndex[i][0]];
+							CurrentParticleIndex += 1;
+						}
+					}
+				}
+
+				TArray<Chaos::TVec3<Chaos::FReal>> VertexTVec3;
+				VertexTVec3.SetNum(Vertex.Num());
+				for (int32 VertexIdx = 0; VertexIdx < Vertex.Num(); VertexIdx++)
+				{
+					VertexTVec3[VertexIdx] = Chaos::TVec3<Chaos::FReal>(Vertex[VertexIdx]);
+				}
+				Chaos::FTriangleMesh TriangleMesh;
+				TriangleMesh.Init(IndicesArray);
+				Chaos::FTriangleMesh::TSpatialHashType<Chaos::FReal> SpatialHash;
+				TConstArrayView<Chaos::TVec3<Chaos::FReal>> ConstArrayViewVertex(VertexTVec3);
+				TriangleMesh.BuildSpatialHash(ConstArrayViewVertex, SpatialHash, Chaos::FReal(SearchRadius));
+				TArray<FIntVector4> NewTetrahedra;
+				auto ComputeDihedralAngle = [](const FVector3f& A, const FVector3f& B, const FVector3f& C, const FVector3f& D)
+					{
+						FVector3f Normal1 = FVector3f::CrossProduct(B - A, C - A).GetSafeNormal();
+						FVector3f Normal2 = FVector3f::CrossProduct(B - A, D - A).GetSafeNormal();
+
+						float CosTheta = FVector3f::DotProduct(Normal1, Normal2);
+						return FMath::RadiansToDegrees(FMath::Acos(CosTheta));
+					};
+				auto IfTetInverted = [](const FVector3f& A, const FVector3f& B, const FVector3f& C, const FVector3f& D)
+					{
+						return (D - A).Dot(FVector3f::CrossProduct(B - A, C - A)) < 0;
+					};
+				int32 NumSkippedDihedral = 0;
+				for (int32 PointIndex : IndicesMap)
+				{
+					TArray<Chaos::TTriangleCollisionPoint<Chaos::FReal>> Result;
+					if (TriangleMesh.PointClosestTriangleQuery(SpatialHash, ConstArrayViewVertex,
+						PointIndex, VertexTVec3[PointIndex], Chaos::FReal(SearchRadius) / 2.f, Chaos::FReal(SearchRadius) / 2.f,
+						[this, &ComponentIndex, &IndicesArray](const int32 PointIndex, const int32 TriangleIndex)->bool
+						{
+							return ComponentIndex[PointIndex] != ComponentIndex[IndicesArray[TriangleIndex][0]];
+						},
+						Result))
+					{
+						for (const Chaos::TTriangleCollisionPoint<Chaos::FReal>& CollisionPoint : Result)
+						{
+							int32 Tri0 = IndicesArray[CollisionPoint.Indices[1]][0];
+							int32 Tri1 = IndicesArray[CollisionPoint.Indices[1]][1];
+							int32 Tri2 = IndicesArray[CollisionPoint.Indices[1]][2];
+							float DihedralAngle = ComputeDihedralAngle(Vertex[PointIndex], Vertex[Tri0], Vertex[Tri1], Vertex[Tri2]);
+							if (DihedralAngle < 10 || DihedralAngle > 170)
+							{
+								NumSkippedDihedral++;
+								continue;
+							}
+							if (IfTetInverted(Vertex[PointIndex], Vertex[Tri0], Vertex[Tri1], Vertex[Tri2]))
+							{
+								Swap(Tri0, Tri1);
+							}
+							NewTetrahedra.Add(FIntVector4(PointIndex,
+								Tri0,
+								Tri1,
+								Tri2));
+							break;					
+						}
+					}
+				}
+				int32 StartSize = InFleshCollection->AddElements(NewTetrahedra.Num(), FTetrahedralCollection::TetrahedralGroup);
+				for (int32 Idx = 0; Idx < NewTetrahedra.Num(); ++Idx)
+				{
+					Tetrahedron[StartSize + Idx] = NewTetrahedra[Idx];
+				}
+				InFleshCollection->InitIncidentElements(); //recompute incident elements after appending new tet constraints
+				UE_LOG(LogChaosFlesh, Warning, TEXT("CreateAirTetrahedralConstraint: Added %d volumetric constraints.")
+					, NewTetrahedra.Num());
+				UE_LOG(LogChaosFlesh, Warning, TEXT("CreateAirTetrahedralConstraint: Skipped %d volumetric constraints due to extreme dihedral angles.")
+					, NumSkippedDihedral);
+				//Draw dynamic mesh of new tet boundary
+				auto FloatVert = [](FVector3d V) { return FVector3f(V.X, V.Y, V.Z); };
+				auto DoubleVert = [](FVector3f V) { return FVector3d(V.X, V.Y, V.Z); };
+				OutDynamicMesh->Reset();
+				FDynamicMesh3& OutDynamicMesh3 = OutDynamicMesh->GetMeshRef();
+				OutDynamicMesh3.EnableAttributes();
+
+				// Compute boundary triangle mesh
+				TArray<FIntVector3> SurfaceElements = UE::Dataflow::GetSurfaceTriangles(NewTetrahedra, /*bKeepInterior = */ false);
+				for (int32 TriangleIndex = 0; TriangleIndex < SurfaceElements.Num(); ++TriangleIndex)
+				{
+					int v0 = OutDynamicMesh3.AppendVertex(DoubleVert(Vertex[SurfaceElements[TriangleIndex][0]]));
+					int v1 = OutDynamicMesh3.AppendVertex(DoubleVert(Vertex[SurfaceElements[TriangleIndex][1]]));
+					int v2 = OutDynamicMesh3.AppendVertex(DoubleVert(Vertex[SurfaceElements[TriangleIndex][2]]));
+					OutDynamicMesh3.AppendTriangle(FIndex3i(v0, v1, v2));
+				}
+				
+				// Compute normals
+				OutDynamicMesh3.EnableVertexNormals(FVector3f(1, 0, 0));
+				FMeshNormals MeshNormals(&OutDynamicMesh3);
+				MeshNormals.ComputeVertexNormals();
+				for (int32 VertexIndex = 0; VertexIndex < OutDynamicMesh3.VertexCount(); ++VertexIndex)
+				{
+					OutDynamicMesh3.SetVertexNormal(VertexIndex, FloatVert(MeshNormals[VertexIndex]));
+				}
+			} // if Indices
+		} // if Vertices
+		SetValue(Context, MoveTemp(static_cast<FManagedArrayCollection&>(*InFleshCollection.Release())), &Collection);
+		SetValue(Context, OutDynamicMesh, &DynamicMesh);
+	}
 }
\ No newline at end of file
diff --git a/Engine/Plugins/Experimental/ChaosFlesh/Source/ChaosFleshNodes/Public/Dataflow/ChaosFleshSetVertexTrianglePositionTargetBindingNode.h b/Engine/Plugins/Experimental/ChaosFlesh/Source/ChaosFleshNodes/Public/Dataflow/ChaosFleshSetVertexTrianglePositionTargetBindingNode.h
--- a/Engine/Plugins/Experimental/ChaosFlesh/Source/ChaosFleshNodes/Public/Dataflow/ChaosFleshSetVertexTrianglePositionTargetBindingNode.h
+++ b/Engine/Plugins/Experimental/ChaosFlesh/Source/ChaosFleshNodes/Public/Dataflow/ChaosFleshSetVertexTrianglePositionTargetBindingNode.h
@@ -9,6 +9,8 @@
 
 #include "ChaosFleshSetVertexTrianglePositionTargetBindingNode.generated.h"
 
+class UDynamicMesh;
+
 /**
  * Create point-triangle weak constraints (springs) between different geometries based on search radius.
  */
@@ -20,8 +22,6 @@ struct FSetVertexTrianglePositionTargetBindingDataflowNode : public FDataflowNod
 	DATAFLOW_NODE_RENDER_TYPE("SurfaceRender",FGeometryCollection::StaticType(),  "Collection")
 
 public:
-	typedef FManagedArrayCollection DataType;
-
 	UPROPERTY(meta = (DataflowInput, DataflowOutput, DisplayName = "Collection", DataflowPassthrough = "Collection"))
 	FManagedArrayCollection Collection;
 
@@ -68,8 +68,6 @@ public:
 	}
 
 private:
-	typedef FManagedArrayCollection DataType;
-
 	UPROPERTY(meta = (DataflowInput, DataflowIntrinsic, DataflowOutput, DisplayName = "Collection", DataflowPassthrough = "Collection"))
 	FManagedArrayCollection Collection;
 
@@ -105,8 +103,6 @@ public:
 	}
 
 private:
-	typedef FManagedArrayCollection DataType;
-
 	UPROPERTY(meta = (DataflowInput, DataflowIntrinsic, DataflowOutput, DisplayName = "Collection", DataflowPassthrough = "Collection"))
 	FManagedArrayCollection Collection;
 
@@ -114,5 +110,51 @@ private:
 	UPROPERTY(meta = (DataflowInput, DataflowIntrinsic, DisplayName = "Collision Vertex Selection"))
 	FDataflowVertexSelection VertexSelection;
 
+	virtual void Evaluate(UE::Dataflow::FContext& Context, const FDataflowOutput* Out) const override;
+};
+
+/**
+ * Create air tetrahedral constraint between point-triangle pair from different geometries based on search radius. 
+ * The added tetrahedra help to maintain distance between geometries.
+ * This node renders the boundary of the added tetrahedral mesh.
+ */
+USTRUCT(meta = (DataflowFlesh))
+struct FCreateAirTetrahedralConstraintDataflowNode : public FDataflowNode
+{
+	GENERATED_USTRUCT_BODY()
+	DATAFLOW_NODE_DEFINE_INTERNAL(FCreateAirTetrahedralConstraintDataflowNode, "CreateAirTetrahedralConstraint", "Flesh", "")
+	DATAFLOW_NODE_RENDER_TYPE("SurfaceRender", FName("FDynamicMesh3"), "DynamicMesh")
+
+public:
+	FCreateAirTetrahedralConstraintDataflowNode(const UE::Dataflow::FNodeParameters& InParam, FGuid InGuid = FGuid::NewGuid())
+		: FDataflowNode(InParam, InGuid)
+	{
+		RegisterInputConnection(&Collection);
+		RegisterInputConnection(&VertexSelection);
+		RegisterOutputConnection(&Collection, &Collection);
+		RegisterOutputConnection(&DynamicMesh)
+			.SetCanHidePin(true)
+			.SetPinIsHidden(true);
+	}
+
+private:
+	UPROPERTY(meta = (DataflowInput, DataflowOutput, DisplayName = "Collection", DataflowPassthrough = "Collection"))
+	FManagedArrayCollection Collection;
+
+	/* (optional) only create tetrahedral constraints from vertices in VertexSelection to triangles in other geometries.
+	* For example, if the VertexSelection contains only one geometry, only this geometry will bind to other geometries.
+	* No constraints will be created between two geometries that are not in the VertexSelection.
+	*/
+	UPROPERTY(meta = (DataflowInput, DisplayName = "(Optional) VertexSelection"))
+	FDataflowVertexSelection VertexSelection;
+
+	/* tetrahedral constraint search radius*/
+	UPROPERTY(EditAnywhere, Category = "Dataflow", meta = (ClampMin = "0.0"))
+	float SearchRadius = 0.f;
+
+	/** Render dynamic mesh of the boundary mesh of added tetrahedra*/
+	UPROPERTY(meta = (DataflowOutput))
+	TObjectPtr<UDynamicMesh> DynamicMesh;
+
 	virtual void Evaluate(UE::Dataflow::FContext& Context, const FDataflowOutput* Out) const override;
 };
\ No newline at end of file
