ChaosFlesh:
-Added origin-insertion range computation across an input animation. This is used as heuristic muscle fiber length threshold.

diff --git a/Engine/Plugins/Experimental/ChaosFlesh/Source/ChaosFleshNodes/Private/Dataflow/ChaosFleshComputeMuscleActivationNode.cpp b/Engine/Plugins/Experimental/ChaosFlesh/Source/ChaosFleshNodes/Private/Dataflow/ChaosFleshComputeMuscleActivationNode.cpp
--- a/Engine/Plugins/Experimental/ChaosFlesh/Source/ChaosFleshNodes/Private/Dataflow/ChaosFleshComputeMuscleActivationNode.cpp
+++ b/Engine/Plugins/Experimental/ChaosFlesh/Source/ChaosFleshNodes/Private/Dataflow/ChaosFleshComputeMuscleActivationNode.cpp
@@ -2,39 +2,49 @@
 
 #include "Dataflow/ChaosFleshComputeMuscleActivationNode.h"
 
-#include "GeometryCollection/Facades/CollectionMeshFacade.h"
-#include "GeometryCollection/Facades/CollectionMuscleActivationFacade.h"
+#include "AnimationRuntime.h"
+#include "Animation/AnimSequence.h"
+#include "Animation/AttributesRuntime.h"
+#include "Animation/Skeleton.h"
+#include "BonePose.h"
+#include "ChaosFlesh/ChaosFlesh.h"
 #include "ChaosFlesh/TetrahedralCollection.h"
-#include "Engine/SkeletalMesh.h"
 #include "Dataflow/DataflowSelection.h"
-#include "ChaosFlesh/ChaosFlesh.h"
+#include "Engine/SkeletalMesh.h"
+#include "GeometryCollection/Facades/CollectionMeshFacade.h"
+#include "GeometryCollection/Facades/CollectionMuscleActivationFacade.h"
+#include "GeometryCollection/Facades/CollectionTransformFacade.h"
+#include "GeometryCollection/Facades/CollectionTransformSourceFacade.h"
+#include "GeometryCollection/Facades/CollectionVertexBoneWeightsFacade.h"
 
 #include UE_INLINE_GENERATED_CPP_BY_NAME(ChaosFleshComputeMuscleActivationNode)
 
 void FComputeMuscleActivationDataNode::Evaluate(UE::Dataflow::FContext& Context, const FDataflowOutput* Out) const
 {
+	using namespace GeometryCollection::Facades;
 	if (Out->IsA(&Collection))
 	{
 		FManagedArrayCollection InCollection = GetValue(Context, &Collection);
 		const TArray<int32>& InOriginIndices = GetValue(Context, &OriginIndicesIn);
 		const TArray<int32>& InInsertionIndices = GetValue(Context, &InsertionIndicesIn);
 		
-		GeometryCollection::Facades::FMuscleActivationFacade FMuscleActivation(InCollection);
-		FMuscleActivation.SetUpMuscleActivation(InOriginIndices, InInsertionIndices, ContractionVolumeScale);
+		FMuscleActivationFacade MuscleActivation(InCollection);
+		MuscleActivation.SetUpMuscleActivation(InOriginIndices, InInsertionIndices, ContractionVolumeScale);
 		Out->SetValue(MoveTemp(InCollection), Context);
 	}
 }
 
 void FComputeMuscleActivationDataNode_v2::Evaluate(UE::Dataflow::FContext& Context, const FDataflowOutput* Out) const
 {
+	using namespace GeometryCollection::Facades;
 	if (Out->IsA(&Collection))
 	{
 		FManagedArrayCollection InCollection = GetValue(Context, &Collection);
 		TArray<int32> InOriginIndices = GetValue(Context, &OriginIndicesIn);
 		TArray<int32> InInsertionIndices = GetValue(Context, &InsertionIndicesIn);
 
-		GeometryCollection::Facades::FMuscleActivationFacade FMuscleActivation(InCollection);
-		if (!FMuscleActivation.SetUpMuscleActivation(InOriginIndices, InInsertionIndices))
+		FMuscleActivationFacade MuscleActivation(InCollection);
+		if (!MuscleActivation.SetUpMuscleActivation(InOriginIndices, InInsertionIndices))
 		{
 			Context.Warning(TEXT("Setup failed, please check the Log for more info."), this, Out);
 		}
@@ -42,6 +52,165 @@ void FComputeMuscleActivationDataNode_v2::Evaluate(UE::Dataflow::FContext& Conte
 	}
 }
 
+bool FindLowestMuscleLengthRatio(const FManagedArrayCollection& InCollection, const UAnimSequence& InAnimationAsset, 
+	const USkeletalMesh& InSkeletalMesh, TArray<float>& MinLengthRatio)
+{
+#if WITH_EDITOR
+	using namespace GeometryCollection::Facades;
+	const TManagedArray<FVector3f>* Vertex = InCollection.FindAttributeTyped<FVector3f>("Vertex", FGeometryCollection::VerticesGroup);
+	if (!Vertex)
+	{
+		UE_LOG(LogChaosFlesh, Error, TEXT("FindLowestMuscleLengthRatio: No Vertex attribute in the Collection."));
+		return false;
+	}
+	// Match transform source skeleton with SkeletalMesh
+	FTransformSource TransformSource(InCollection);
+	if (!TransformSource.IsValid())
+	{
+		UE_LOG(LogChaosFlesh, Error, TEXT("FindLowestMuscleLengthRatio: TransformSource group in the Collection is not valid."));
+		return false;
+	}
+	TSet<int32> Roots;
+	const USkeleton* Skeleton = InSkeletalMesh.GetSkeleton();
+	const FReferenceSkeleton& ReferenceSkeleton = InSkeletalMesh.GetRefSkeleton();
+	if (Skeleton)
+	{
+		Roots = TransformSource.GetTransformSource(Skeleton->GetName(), Skeleton->GetGuid().ToString(), InSkeletalMesh.GetName());
+		if (Roots.IsEmpty())
+		{
+			UE_LOG(LogChaosFlesh, Error, TEXT("FindLowestMuscleLengthRatio: The root in the TransformSource group is incompatible with the SkeletalMesh."));
+			return false;
+		}
+		else
+		{
+			ensureMsgf(Roots.Num() == 1, TEXT("Only supports a single root per skeleton.(%s)"), *Skeleton->GetName());
+		}
+	}
+	else
+	{
+		UE_LOG(LogChaosFlesh, Error, TEXT("FindLowestMuscleLengthRatio: There's no skeleton in the SkeletalMesh."));
+		return false;
+	}
+	const int32 RootTransformOffset = Roots.Array()[0];
+	auto RootShift = [&RootTransformOffset](int32 Transform) { return Transform - RootTransformOffset; };
+	// Rest transforms
+	FCollectionTransformFacade TransformFacade(InCollection);
+	TArray<FTransform> RestTransforms = TransformFacade.ComputeCollectionSpaceTransforms();
+
+	// For extracting pose transforms
+	FMemMark Mark(FMemStack::Get());
+	const int32 NumBones = ReferenceSkeleton.GetNum();
+	TArray<FBoneIndexType> BoneIndices;
+	BoneIndices.SetNumUninitialized(NumBones);
+	for (int32 Index = 0; Index < NumBones; ++Index)
+	{
+		int32 SkeletonBoneIndex = Skeleton->GetReferenceSkeleton().FindBoneIndex(ReferenceSkeleton.GetBoneName(Index));
+		BoneIndices[Index] = StaticCast<FBoneIndexType>(SkeletonBoneIndex);
+	}
+
+	FBoneContainer BoneContainer;
+	BoneContainer.SetUseRAWData(true);
+	BoneContainer.InitializeTo(BoneIndices, UE::Anim::FCurveFilterSettings(), *Skeleton);
+
+	FCompactPose CompactPose;
+	CompactPose.SetBoneContainer(&BoneContainer);
+
+	FBlendedCurve BlendedCurve;
+	BlendedCurve.InitFrom(BoneContainer);
+
+	UE::Anim::FStackAttributeContainer TempAttributes;
+	FAnimationPoseData AnimationPoseData(CompactPose, BlendedCurve, TempAttributes);
+
+	// Prepare kinematic origin insertion weights
+	FVertexBoneWeightsFacade WeightsFacade(InCollection);
+	if (!WeightsFacade.IsValid())
+	{
+		UE_LOG(LogChaosFlesh, Error, TEXT("FindLowestMuscleLengthRatio: No kinematics information in the Collection."));
+		return false;
+	}
+	const FMuscleActivationFacade MuscleActivation(InCollection);
+	if (!MuscleActivation.IsValid())
+	{
+		UE_LOG(LogChaosFlesh, Error, TEXT("FindLowestMuscleLengthRatio: No muscle activation information in the Collection."));
+		return false;
+	}
+	const int32 NumMuscles = MuscleActivation.NumMuscles();
+	TArray<TArray<int32>> OriginBoneIndices, InsertionBoneIndices;
+	TArray<TArray<float>> OriginBoneWeights, InsertionBoneWeights;
+	TArray<FVector3d> OriginPosition, InsertionPosition;
+	TArray<float> OIRestLength;
+	OriginBoneIndices.SetNum(NumMuscles);
+	InsertionBoneIndices.SetNum(NumMuscles);
+	OriginBoneWeights.SetNum(NumMuscles);
+	InsertionBoneWeights.SetNum(NumMuscles);
+	OriginPosition.SetNum(NumMuscles);
+	InsertionPosition.SetNum(NumMuscles);
+	OIRestLength.SetNum(NumMuscles);
+	MinLengthRatio.Init(FLT_MAX, NumMuscles);
+	auto ChaosVert = [](FVector3f V) { return Chaos::FVec3(V.X, V.Y, V.Z); };
+	auto DoubleVert = [](FVector3f V) { return FVector3d(V.X, V.Y, V.Z); };
+	for (int32 MuscleIdx = 0; MuscleIdx < NumMuscles; ++MuscleIdx)
+	{
+		const FMuscleActivationData MuscleActivationData = MuscleActivation.GetMuscleActivationData(MuscleIdx);
+		const int32 OriginIdx = MuscleActivationData.OriginInsertionPair[0];
+		const int32 InsertionIdx = MuscleActivationData.OriginInsertionPair[1];
+		if (ensureMsgf(WeightsFacade.GetBoneIndices().IsValidIndex(OriginIdx), TEXT("Origin index %d is not a valid index in WeightsFacade BoneIndices"), OriginIdx) &&
+			ensureMsgf(WeightsFacade.GetBoneIndices().IsValidIndex(InsertionIdx), TEXT("Insertion index %d is not a valid index in WeightsFacade BoneIndices"), InsertionIdx) &&
+			ensureMsgf(Vertex->IsValidIndex(OriginIdx), TEXT("Origin index %d is not a valid index in Vertex"), OriginIdx) &&
+			ensureMsgf(Vertex->IsValidIndex(InsertionIdx), TEXT("Insertion index %d is not a valid index in Vertex"), InsertionIdx))
+		{
+			OriginBoneIndices[MuscleIdx] = WeightsFacade.GetBoneIndices()[OriginIdx];
+			InsertionBoneIndices[MuscleIdx] = WeightsFacade.GetBoneIndices()[InsertionIdx];
+			OriginBoneWeights[MuscleIdx] = WeightsFacade.GetBoneWeights()[OriginIdx];
+			InsertionBoneWeights[MuscleIdx] = WeightsFacade.GetBoneWeights()[InsertionIdx];
+			OriginPosition[MuscleIdx] = DoubleVert((*Vertex)[OriginIdx]);
+			InsertionPosition[MuscleIdx] = DoubleVert((*Vertex)[InsertionIdx]);
+			OIRestLength[MuscleIdx] = ((*Vertex)[OriginIdx] - (*Vertex)[InsertionIdx]).Size();
+		}
+	}
+
+	for (int32 Frame = 0; Frame < InAnimationAsset.GetNumberOfSampledKeys(); ++Frame)
+	{
+		const FAnimExtractContext ExtractionContext(double(InAnimationAsset.GetTimeAtFrame(Frame)));
+		InAnimationAsset.GetAnimationPose(AnimationPoseData, ExtractionContext);
+		TArray<FTransform> ComponentSpaceTransforms;
+		FAnimationRuntime::FillUpComponentSpaceTransforms(ReferenceSkeleton, AnimationPoseData.GetPose().GetBones(), ComponentSpaceTransforms);
+
+		for (int32 MuscleIdx = 0; MuscleIdx < NumMuscles; ++MuscleIdx)
+		{
+			FVector3d OriginPos(0.f), InsertionPos(0.f);
+			for (int32 InfluenceIdx = 0; InfluenceIdx < OriginBoneIndices[MuscleIdx].Num(); ++InfluenceIdx)
+			{
+				const int32 OBoneIdx = OriginBoneIndices[MuscleIdx][InfluenceIdx];
+				const float OBoneWeight = OriginBoneWeights[MuscleIdx][InfluenceIdx];
+				const int32 ShiftedOBoneIdx = RootShift(OBoneIdx);
+				if (RestTransforms.IsValidIndex(OBoneIdx) &&
+					ComponentSpaceTransforms.IsValidIndex(ShiftedOBoneIdx))
+				{
+					OriginPos += ComponentSpaceTransforms[ShiftedOBoneIdx].TransformPosition(RestTransforms[OBoneIdx].InverseTransformPosition(OriginPosition[MuscleIdx])) * OBoneWeight;
+				}
+			}
+			for (int32 InfluenceIdx = 0; InfluenceIdx < InsertionBoneIndices[MuscleIdx].Num(); ++InfluenceIdx)
+			{
+				const int32 IBoneIdx = InsertionBoneIndices[MuscleIdx][InfluenceIdx];
+				const float IBoneWeight = InsertionBoneWeights[MuscleIdx][InfluenceIdx];
+				const int32 ShiftedIBoneIdx = RootShift(IBoneIdx);
+				if (RestTransforms.IsValidIndex(IBoneIdx) &&
+					ComponentSpaceTransforms.IsValidIndex(ShiftedIBoneIdx))
+				{
+					InsertionPos += ComponentSpaceTransforms[ShiftedIBoneIdx].TransformPosition(RestTransforms[IBoneIdx].InverseTransformPosition(InsertionPosition[MuscleIdx])) * IBoneWeight;
+				}
+			}
+			float Ratio = (OriginPos - InsertionPos).Size() / OIRestLength[MuscleIdx];
+			MinLengthRatio[MuscleIdx] = FMath::Min(MinLengthRatio[MuscleIdx], Ratio);
+		}
+	}
+	return true;
+#else
+	return false;
+#endif
+}
+
 FSetMuscleActivationParameterNode::FSetMuscleActivationParameterNode(const UE::Dataflow::FNodeParameters& InParam, FGuid InGuid)
 	: FDataflowNode(InParam, InGuid)
 	, ApplyGlobalParameters(
@@ -50,19 +219,77 @@ FSetMuscleActivationParameterNode::FSetMuscleActivationParameterNode(const UE::D
 				// triggers node invalidation
 				ParameterMethod = EParameterMethod::Global;
 			}))
+	, ImportLowestMuscleLengthRatio(
+		FDataflowFunctionProperty::FDelegate::CreateLambda([this](UE::Dataflow::FContext& Context)
+			{
+				TObjectPtr<UAnimSequence> InAnimationAsset = GetValue(Context, &AnimationAsset);
+				TObjectPtr<USkeletalMesh> InSkeletalMesh = GetValue(Context, &SkeletalMesh);
+				if (InAnimationAsset && InSkeletalMesh)
+				{
+					const FManagedArrayCollection& InCollection = GetValue(Context, &Collection);
+					GeometryCollection::Facades::FMuscleActivationFacade MuscleActivation(InCollection);
+					TArray<float> MinLengthRatio;
+					if (FindLowestMuscleLengthRatio(InCollection, *InAnimationAsset, *InSkeletalMesh, MinLengthRatio)) //success
+					{
+						const int32 NumMuscles = MinLengthRatio.Num();
+						for (int32 Idx = 0; Idx < ParameterArray.Num(); ++Idx)
+						{
+							const FString MuscleName = ParameterArray[Idx].MuscleName;
+							const int32 MuscleIndex = MuscleActivation.FindMuscleIndexByName(MuscleName);
+							if (MinLengthRatio.IsValidIndex(MuscleIndex))
+							{
+								if (MinLengthRatio[MuscleIndex] > 0 && MinLengthRatio[MuscleIndex] < 1)
+								{
+									ParameterArray[Idx].MuscleLengthRatioThresholdForMaxActivation = MinLengthRatio[MuscleIndex];
+								}
+								else
+								{
+									UE_LOG(LogChaosFlesh, Warning,
+										TEXT("SetMuscleActivationParameter::ImportLowestMuscleLengthRatio: Muscle [%s] index [%d] has minimum origin-insertion length ratio %.2f across the whole animation."),
+										*MuscleName, MuscleIndex, MinLengthRatio[MuscleIndex]);
+								}
+							}
+							else
+							{
+								UE_LOG(LogChaosFlesh, Error,
+									TEXT("SetMuscleActivationParameter::ImportLowestMuscleLengthRatio: Geometry [%s] is not a valid muscle."),
+									*MuscleName);
+								Context.Error(TEXT("ImportLowestMuscleLengthRatio: Geometry [%s] is not a valid muscle."), this);
+							}
+						}
+					}
+					else
+					{
+						Context.Error(TEXT("FindLowestMuscleLengthRatio failed, please check the Log for more info."), this);
+					}
+				}
+			}))
 	, ImportAllMuscleNames(
 		FDataflowFunctionProperty::FDelegate::CreateLambda([this](UE::Dataflow::FContext& Context)
 			{
-				const FManagedArrayCollection& InCollection = GetValue<FManagedArrayCollection>(Context, &Collection);
-				const GeometryCollection::Facades::FMuscleActivationFacade FMuscleActivation(InCollection);
-				ParameterArray.SetNum(FMuscleActivation.NumMuscles());
-				for (int32 MuscleIdx = 0; MuscleIdx < FMuscleActivation.NumMuscles(); ++MuscleIdx)
+				using namespace GeometryCollection::Facades;
+				const FManagedArrayCollection& InCollection = GetValue(Context, &Collection);
+				const FMuscleActivationFacade MuscleActivation(InCollection);
+				ParameterArray.SetNum(MuscleActivation.NumMuscles());
+				for (int32 MuscleIdx = 0; MuscleIdx < MuscleActivation.NumMuscles(); ++MuscleIdx)
 				{
-					ParameterArray[MuscleIdx].MuscleName = FMuscleActivation.FindMuscleName(MuscleIdx);
-					ParameterArray[MuscleIdx].ContractionVolumeScale = ContractionVolumeScale;
-					ParameterArray[MuscleIdx].FiberLengthRatioAtMaxActivation = GlobalFiberLengthRatioAtMaxActivation;
-					ParameterArray[MuscleIdx].MuscleLengthRatioThresholdForMaxActivation = GlobalMuscleLengthRatioThresholdForMaxActivation;
-					ParameterArray[MuscleIdx].InflationVolumeScale = GlobalInflationVolumeScale;
+					ParameterArray[MuscleIdx].MuscleName = MuscleActivation.FindMuscleName(MuscleIdx);
+					// load existing attributes
+					if (MuscleActivation.IsValid())
+					{
+						const FMuscleActivationData MuscleActivationData = MuscleActivation.GetMuscleActivationData(MuscleIdx);
+						ParameterArray[MuscleIdx].ContractionVolumeScale = MuscleActivationData.ContractionVolumeScale.Num() ? MuscleActivationData.ContractionVolumeScale[0] : ContractionVolumeScale;
+						ParameterArray[MuscleIdx].FiberLengthRatioAtMaxActivation = MuscleActivationData.FiberLengthRatioAtMaxActivation;
+						ParameterArray[MuscleIdx].MuscleLengthRatioThresholdForMaxActivation = MuscleActivationData.MuscleLengthRatioThresholdForMaxActivation;
+						ParameterArray[MuscleIdx].InflationVolumeScale = MuscleActivationData.InflationVolumeScale;
+					}
+					else
+					{
+						ParameterArray[MuscleIdx].ContractionVolumeScale = ContractionVolumeScale;
+						ParameterArray[MuscleIdx].FiberLengthRatioAtMaxActivation = GlobalFiberLengthRatioAtMaxActivation;
+						ParameterArray[MuscleIdx].MuscleLengthRatioThresholdForMaxActivation = GlobalMuscleLengthRatioThresholdForMaxActivation;
+						ParameterArray[MuscleIdx].InflationVolumeScale = GlobalInflationVolumeScale;
+					}
 				}
 			}))
 	, ResetToGlobalParameters(
@@ -84,18 +311,26 @@ FSetMuscleActivationParameterNode::FSetMuscleActivationParameterNode(const UE::D
 			}))
 {
 	RegisterInputConnection(&Collection);
+	RegisterInputConnection(&AnimationAsset)
+		.SetCanHidePin(true)
+		.SetPinIsHidden(true);
+	RegisterInputConnection(&SkeletalMesh)
+		.SetCanHidePin(true)
+		.SetPinIsHidden(true);
 	RegisterOutputConnection(&Collection, &Collection);
 }
 
+
 void FSetMuscleActivationParameterNode::Evaluate(UE::Dataflow::FContext& Context, const FDataflowOutput* Out) const
 {
+	using namespace GeometryCollection::Facades;
 	if (Out->IsA(&Collection))
 	{
 		FManagedArrayCollection InCollection = GetValue(Context, &Collection);
 		if (IsConnected(&Collection))
 		{
-			GeometryCollection::Facades::FMuscleActivationFacade FMuscleActivation(InCollection);
-			FMuscleActivation.UpdateGlobalMuscleActivationParameters(
+			FMuscleActivationFacade MuscleActivation(InCollection);
+			MuscleActivation.UpdateGlobalMuscleActivationParameters(
 					ContractionVolumeScale,
 					GlobalFiberLengthRatioAtMaxActivation,
 					GlobalMuscleLengthRatioThresholdForMaxActivation,
@@ -106,10 +341,10 @@ void FSetMuscleActivationParameterNode::Evaluate(UE::Dataflow::FContext& Context
 				for (int32 Idx = 0; Idx < ParameterArray.Num(); ++Idx)
 				{
 					FPerMuscleParameter Params = ParameterArray[Idx];
-					int32 MuscleIndex = FMuscleActivation.FindMuscleIndexByName(Params.MuscleName);
-					if (FMuscleActivation.IsValidMuscleIndex(MuscleIndex))
+					int32 MuscleIndex = MuscleActivation.FindMuscleIndexByName(Params.MuscleName);
+					if (MuscleActivation.IsValidMuscleIndex(MuscleIndex))
 					{
-						FMuscleActivation.UpdateMuscleActivationParameters(
+						MuscleActivation.UpdateMuscleActivationParameters(
 							MuscleIndex,
 							Params.ContractionVolumeScale,
 							Params.FiberLengthRatioAtMaxActivation,
@@ -176,10 +411,11 @@ FReadSkeletalMeshCurvesDataflowNode::FReadSkeletalMeshCurvesDataflowNode(const U
 
 void FReadSkeletalMeshCurvesDataflowNode::Evaluate(UE::Dataflow::FContext& Context, const FDataflowOutput* Out) const
 {
+	using namespace GeometryCollection::Facades;
 	if (Out->IsA(&Collection) || Out->IsA(&GeometrySelection))
 	{
 		FManagedArrayCollection InCollection = GetValue(Context, &Collection);
-		GeometryCollection::Facades::FMuscleActivationFacade MuscleActivationFacade(InCollection);
+		FMuscleActivationFacade MuscleActivationFacade(InCollection);
 		TArray<int32> MuscleCurveGeometry;
 		if (MuscleActivationFacade.IsValid())
 		{
diff --git a/Engine/Plugins/Experimental/ChaosFlesh/Source/ChaosFleshNodes/Public/Dataflow/ChaosFleshComputeMuscleActivationNode.h b/Engine/Plugins/Experimental/ChaosFlesh/Source/ChaosFleshNodes/Public/Dataflow/ChaosFleshComputeMuscleActivationNode.h
--- a/Engine/Plugins/Experimental/ChaosFlesh/Source/ChaosFleshNodes/Public/Dataflow/ChaosFleshComputeMuscleActivationNode.h
+++ b/Engine/Plugins/Experimental/ChaosFlesh/Source/ChaosFleshNodes/Public/Dataflow/ChaosFleshComputeMuscleActivationNode.h
@@ -9,6 +9,8 @@
 
 #include "ChaosFleshComputeMuscleActivationNode.generated.h"
 
+class UAnimSequence;
+
 /**
 * Please use new version of this node, then use SetMuscleActivationParameter node to set up muscle activation parameters.
 * Computes an orthogonal matrix for each element
@@ -86,7 +88,7 @@ struct FPerMuscleParameter
 {
 	GENERATED_USTRUCT_BODY()
 	UPROPERTY(EditAnywhere, meta = (InlineEditConditionToggle), Category = "Muscle Activation Parameters")
-	bool bCanEditMuscleName;
+	bool bCanEditMuscleName = false;
 	UPROPERTY(EditAnywhere, meta = (EditCondition = "bCanEditMuscleName"), Category = "Muscle Activation Parameters")
 	FString MuscleName;
 	UPROPERTY(EditAnywhere, Category = "Muscle Activation Parameters", meta = (ClampMin = "0"))
@@ -143,6 +145,20 @@ public:
 	UPROPERTY(EditAnywhere, Category = "Global Parameters", meta = (ClampMin = "0"))
 	float GlobalInflationVolumeScale = 1.f;
 	
+	/* Use minimum muscle lengths across this animation asset to infer MuscleLengthRatioThresholdForMaxActivation*/
+	UPROPERTY(EditAnywhere, Category = "Animation Data", meta = (DataflowInput))
+	TObjectPtr<UAnimSequence> AnimationAsset = nullptr;
+
+	/* Skeletal mesh used to linear blend skin kinematic origins and insertions. 
+	* This must be the same as the one used to create TransformSource group. 
+	* Check geometry spreadsheet for more info.*/
+	UPROPERTY(EditAnywhere, Category = "Animation Data", meta = (DataflowInput))
+	TObjectPtr<USkeletalMesh> SkeletalMesh = nullptr;
+
+	/** Click on this button to import the minimum muscle lengths across the animation asset as MuscleLengthRatioThresholdForMaxActivation*/
+	UPROPERTY(EditAnywhere, Category = "Animation Data", meta = (ButtonImage = "Icons.Import"))
+	FDataflowFunctionProperty ImportLowestMuscleLengthRatio;
+
 	/** Click on this button to import muscle bone names (in the Transform group) from the collection. */
 	UPROPERTY(EditAnywhere, Category = "Custom Parameters", meta = (ButtonImage = "Icons.Import"))
 	FDataflowFunctionProperty ImportAllMuscleNames;
