ChaosFlesh:
-Support static mesh embedding and geom cache generation

diff --git a/Engine/Plugins/Experimental/ChaosFlesh/Source/ChaosFleshEngine/Private/ChaosFlesh/ChaosDeformableTetrahedralComponent.cpp b/Engine/Plugins/Experimental/ChaosFlesh/Source/ChaosFleshEngine/Private/ChaosFlesh/ChaosDeformableTetrahedralComponent.cpp
--- a/Engine/Plugins/Experimental/ChaosFlesh/Source/ChaosFleshEngine/Private/ChaosFlesh/ChaosDeformableTetrahedralComponent.cpp
+++ b/Engine/Plugins/Experimental/ChaosFlesh/Source/ChaosFleshEngine/Private/ChaosFlesh/ChaosDeformableTetrahedralComponent.cpp
@@ -712,14 +712,36 @@ TArray<FVector3f> UDeformableTetrahedralComponent::GetGeometryCachePositions(con
 	return Positions;
 }
 
+TArray<FVector3f> UDeformableTetrahedralComponent::GetGeometryCachePositions(const UStaticMesh* StaticMesh) const
+{
+	TArray<FVector3f> Positions;
+	if (RestCollection && DynamicCollection && StaticMesh)
+	{
+		const FFleshCollection* FleshCollection = RestCollection->GetCollection();
+		const TManagedArray<FVector3f>* RestVertices = RestCollection->FindPositions();
+		const TManagedArray<FVector3f>* SimulatedVertices = DynamicCollection->FindPositions();
+		if (FleshCollection && RestVertices && SimulatedVertices)
+		{
+			ChaosFlesh::BoundSurfacePositions(StaticMesh, FleshCollection, RestVertices, SimulatedVertices, Positions);
+		}
+	}
+	return Positions;
+}
+
 TOptional<TArray<int32>> UDeformableTetrahedralComponent::GetMeshImportVertexMap(const USkinnedAsset& SkinnedMeshAsset) const
 {
 	constexpr int32 LODIndex = 0;
 	const TOptional<TArray<int32>> None;
 #if WITH_EDITOR
 	const FSkeletalMeshModel* const MLDModel = SkinnedMeshAsset.GetImportedModel();
-	if (!MLDModel || !MLDModel->LODModels.IsValidIndex(LODIndex))
+	if (!MLDModel)
+	{
+		UE_LOG(LogDeformableTetrahedralComponentInternal, Warning, TEXT("MLDModel is empty."));
+		return None;
+	}
+	if (!MLDModel->LODModels.IsValidIndex(LODIndex))
 	{
+		UE_LOG(LogDeformableTetrahedralComponentInternal, Warning, TEXT("LODIndex %d is not a valid index."), LODIndex);
 		return None;
 	}
 	const FSkeletalMeshLODModel& MLDLOD = MLDModel->LODModels[LODIndex];
diff --git a/Engine/Plugins/Experimental/ChaosFlesh/Source/ChaosFleshEngine/Private/ChaosFlesh/FleshCollectionEngineUtility.cpp b/Engine/Plugins/Experimental/ChaosFlesh/Source/ChaosFleshEngine/Private/ChaosFlesh/FleshCollectionEngineUtility.cpp
--- a/Engine/Plugins/Experimental/ChaosFlesh/Source/ChaosFleshEngine/Private/ChaosFlesh/FleshCollectionEngineUtility.cpp
+++ b/Engine/Plugins/Experimental/ChaosFlesh/Source/ChaosFleshEngine/Private/ChaosFlesh/FleshCollectionEngineUtility.cpp
@@ -33,14 +33,13 @@ namespace ChaosFlesh
 	}
 
 	void BoundSurfacePositions(
-		const USkeletalMesh* SkeletalMesh,
+		const FString& MeshId,
 		const FFleshCollection* FleshCollection,
 		const TManagedArray<FVector3f>* RestVertices,
 		const TManagedArray<FVector3f>* SimulatedVertices,
-		TArray<FVector3f>& Positions)
+		TArray<FVector3f>& OutPositions)
 	{
 		GeometryCollection::Facades::FTetrahedralBindings TetBindings(*FleshCollection);
-		FString MeshId = GetMeshId(SkeletalMesh, false);
 		FName MeshIdName(MeshId);
 
 		const int32 LODIndex = 0;
@@ -65,8 +64,11 @@ namespace ChaosFlesh
 		if (!BindingsEval.IsValid())
 		{
 			UE_LOG(LogChaosFlesh, Error,
-				TEXT("CreateGeometryCache - Invalid flesh bindings for skeletal mesh asset [%s]"),
-				*SkeletalMesh->GetName());
+				TEXT("CreateGeometryCache - Min parent index = %d at tet index %d"),
+				BindingsEval.MinIndex(), BindingsEval.MinIndexPosition());
+			UE_LOG(LogChaosFlesh, Error,
+				TEXT("CreateGeometryCache - Invalid flesh bindings for mesh [%s]"),
+				*MeshId);
 			return;
 		}
 
@@ -78,11 +80,31 @@ namespace ChaosFlesh
 		}
 
 		const int32 NumVertices = BindingsEval.NumVertices();
-		Positions.SetNum(NumVertices);
+		OutPositions.SetNum(NumVertices);
 		for (int32 Index = 0; Index < NumVertices; ++Index)
 		{
-			Positions[Index] = BindingsEval.GetEmbeddedPosition(Index, CurrVertices);
+			OutPositions[Index] = BindingsEval.GetEmbeddedPosition(Index, CurrVertices);
 		}
-		
+	}
+	void BoundSurfacePositions(
+		const USkeletalMesh* SkeletalMesh,
+		const FFleshCollection* FleshCollection,
+		const TManagedArray<FVector3f>* RestVertices,
+		const TManagedArray<FVector3f>* SimulatedVertices,
+		TArray<FVector3f>& OutPositions)
+	{
+		FString MeshId = GetMeshId(SkeletalMesh, false);
+		BoundSurfacePositions(MeshId, FleshCollection, RestVertices, SimulatedVertices, OutPositions);
+	}
+
+	void BoundSurfacePositions(
+		const UStaticMesh* StaticMesh,
+		const FFleshCollection* FleshCollection,
+		const TManagedArray<FVector3f>* RestVertices,
+		const TManagedArray<FVector3f>* SimulatedVertices,
+		TArray<FVector3f>& OutPositions)
+	{
+		FString MeshId = GetMeshId(StaticMesh);
+		BoundSurfacePositions(MeshId, FleshCollection, RestVertices, SimulatedVertices, OutPositions);
 	}
 }
diff --git a/Engine/Plugins/Experimental/ChaosFlesh/Source/ChaosFleshEngine/Public/ChaosFlesh/ChaosDeformableTetrahedralComponent.h b/Engine/Plugins/Experimental/ChaosFlesh/Source/ChaosFleshEngine/Public/ChaosFlesh/ChaosDeformableTetrahedralComponent.h
--- a/Engine/Plugins/Experimental/ChaosFlesh/Source/ChaosFleshEngine/Public/ChaosFlesh/ChaosDeformableTetrahedralComponent.h
+++ b/Engine/Plugins/Experimental/ChaosFlesh/Source/ChaosFleshEngine/Public/ChaosFlesh/ChaosDeformableTetrahedralComponent.h
@@ -140,6 +140,7 @@ public:
 	
 	//~ Begin IDataflowGeometryCachable Interface
 	virtual TArray<FVector3f> GetGeometryCachePositions(const USkeletalMesh* SkeletalMesh) const override;
+	virtual TArray<FVector3f> GetGeometryCachePositions(const UStaticMesh* StaticMesh) const override;
 
 	virtual TOptional<TArray<int32>> GetMeshImportVertexMap(const USkinnedAsset& SkinnedMeshAsset) const override;
 
diff --git a/Engine/Plugins/Experimental/ChaosFlesh/Source/ChaosFleshEngine/Public/ChaosFlesh/FleshCollectionEngineUtility.h b/Engine/Plugins/Experimental/ChaosFlesh/Source/ChaosFleshEngine/Public/ChaosFlesh/FleshCollectionEngineUtility.h
--- a/Engine/Plugins/Experimental/ChaosFlesh/Source/ChaosFleshEngine/Public/ChaosFlesh/FleshCollectionEngineUtility.h
+++ b/Engine/Plugins/Experimental/ChaosFlesh/Source/ChaosFleshEngine/Public/ChaosFlesh/FleshCollectionEngineUtility.h
@@ -15,10 +15,24 @@ namespace ChaosFlesh
 
 	CHAOSFLESHENGINE_API FString GetMeshId(const UStaticMesh* StaticMesh);
 
+	void BoundSurfacePositions(
+		const FString& MeshId,
+		const FFleshCollection* FleshCollection,
+		const TManagedArray<FVector3f>* RestVertices,
+		const TManagedArray<FVector3f>* SimulatedVertices,
+		TArray<FVector3f>& OutPositions);
+
 	CHAOSFLESHENGINE_API void BoundSurfacePositions(
 		const USkeletalMesh* SkeletalMesh,
 		const FFleshCollection* FleshCollection,
 		const TManagedArray<FVector3f>* RestVertices,
 		const TManagedArray<FVector3f>* SimulatedVertices,
-		TArray<FVector3f>& Positions);
+		TArray<FVector3f>& OutPositions);
+
+	CHAOSFLESHENGINE_API void BoundSurfacePositions(
+		const UStaticMesh* StaticMesh,
+		const FFleshCollection* FleshCollection,
+		const TManagedArray<FVector3f>* RestVertices,
+		const TManagedArray<FVector3f>* SimulatedVertices,
+		TArray<FVector3f>& OutPositions);
 }
diff --git a/Engine/Plugins/Experimental/ChaosFlesh/Source/ChaosFleshNodes/ChaosFleshNodes.Build.cs b/Engine/Plugins/Experimental/ChaosFlesh/Source/ChaosFleshNodes/ChaosFleshNodes.Build.cs
--- a/Engine/Plugins/Experimental/ChaosFlesh/Source/ChaosFleshNodes/ChaosFleshNodes.Build.cs
+++ b/Engine/Plugins/Experimental/ChaosFlesh/Source/ChaosFleshNodes/ChaosFleshNodes.Build.cs
@@ -21,8 +21,9 @@ namespace UnrealBuildTool.Rules
 					"ChaosFleshEngine",
 					"Engine",
 					"GeometryAlgorithms",
-					"GeometryFramework",
+					"GeometryCollectionEngine",
 					"GeometryCore",
+					"GeometryFramework",
 					"MeshConversion",
 					"MeshConversionEngineTypes",
 					"MeshDescription",
diff --git a/Engine/Plugins/Experimental/ChaosFlesh/Source/ChaosFleshNodes/Private/Dataflow/ChaosFleshGenerateSurfaceBindingsNode.cpp b/Engine/Plugins/Experimental/ChaosFlesh/Source/ChaosFleshNodes/Private/Dataflow/ChaosFleshGenerateSurfaceBindingsNode.cpp
--- a/Engine/Plugins/Experimental/ChaosFlesh/Source/ChaosFleshNodes/Private/Dataflow/ChaosFleshGenerateSurfaceBindingsNode.cpp
+++ b/Engine/Plugins/Experimental/ChaosFlesh/Source/ChaosFleshNodes/Private/Dataflow/ChaosFleshGenerateSurfaceBindingsNode.cpp
@@ -14,6 +14,7 @@
 #include "Engine/StaticMesh.h"
 #include "Engine/SkeletalMesh.h"
 #include "GeometryCollection/Facades/CollectionTetrahedralBindingsFacade.h"
+#include "GeometryCollection/GeometryCollectionEngineConversion.h"
 #include "IndexTypes.h"
 #include "MeshDescription.h"
 #include "MeshDescriptionToDynamicMesh.h"
@@ -33,28 +34,52 @@ DEFINE_LOG_CATEGORY(LogMeshBindings);
 void
 BuildVertexToVertexAdjacencyBuffer(
 	const UE::Geometry::FDynamicMesh3 DynamicMesh,
-	TArray<TArray<uint32>>& NeighborNodes)
+	TArray<TArray<uint32>>& OutNeighborNodes)
 {
-	NeighborNodes.SetNum(DynamicMesh.VertexCount());
-	for (int i = 0; i < DynamicMesh.TriangleCount(); i++)
+	OutNeighborNodes.SetNum(DynamicMesh.VertexCount());
+	for (int32 TriIdx = 0; TriIdx < DynamicMesh.TriangleCount(); ++TriIdx)
 	{
-		const UE::Geometry::FIndex3i& Tri = DynamicMesh.GetTriangle(i);
-		TArray<uint32>& N0 = NeighborNodes[Tri[0]];
+		const UE::Geometry::FIndex3i& Tri = DynamicMesh.GetTriangle(TriIdx);
+		TArray<uint32>& N0 = OutNeighborNodes[Tri[0]];
 		N0.AddUnique(Tri[1]);
 		N0.AddUnique(Tri[2]);
-		TArray<uint32>& N1 = NeighborNodes[Tri[1]];
+		TArray<uint32>& N1 = OutNeighborNodes[Tri[1]];
 		N1.AddUnique(Tri[0]);
 		N1.AddUnique(Tri[2]);
-		TArray<uint32>& N2 = NeighborNodes[Tri[2]];
+		TArray<uint32>& N2 = OutNeighborNodes[Tri[2]];
 		N2.AddUnique(Tri[0]);
 		N2.AddUnique(Tri[1]);
 	}
 }
 
+void UnloadMeshDescription(const FMeshDescription& SourceMesh,
+	TArray<FVector3f>& OutVertices, TArray<FIntVector>& OutTriangles, TArray<TArray<uint32>>& OutNeighborNodes)
+{
+	FDynamicMesh3 DynamicMesh;
+	FMeshDescriptionToDynamicMesh Converter;
+	Converter.Convert(&SourceMesh, DynamicMesh);
+
+	OutVertices.SetNumUninitialized(DynamicMesh.VertexCount());
+	for (int32 VertexIdx = 0; VertexIdx < DynamicMesh.VertexCount(); ++VertexIdx)
+	{
+		const FVector3d& Pos = DynamicMesh.GetVertex(VertexIdx);
+		OutVertices[VertexIdx].Set(Pos[0], Pos[1], Pos[2]);
+	}
+
+	OutTriangles.SetNumUninitialized(DynamicMesh.TriangleCount());
+	for (int32 TriIdx = 0; TriIdx < DynamicMesh.TriangleCount(); ++TriIdx)
+	{
+		const UE::Geometry::FIndex3i& Tri = DynamicMesh.GetTriangle(TriIdx);
+		OutTriangles[TriIdx] = FIntVector(Tri[0], Tri[1], Tri[2]);
+	}
+
+	BuildVertexToVertexAdjacencyBuffer(DynamicMesh, OutNeighborNodes);
+}
+
 void
 BuildVertexToVertexAdjacencyBuffer(
 	const FSkeletalMeshLODRenderData& LodRenderData,
-	TArray<TArray<uint32>>& NeighborNodes)
+	TArray<TArray<uint32>>& OutNeighborNodes)
 {
 	const FRawStaticIndexBuffer16or32Interface* IndexBuffer = LodRenderData.MultiSizeIndexContainer.GetIndexBuffer();
 	const uint32 IndexCount = IndexBuffer->Num();
@@ -62,8 +87,8 @@ BuildVertexToVertexAdjacencyBuffer(
 	const FPositionVertexBuffer& VertexBuffer = LodRenderData.StaticVertexBuffers.PositionVertexBuffer;
 	const uint32 VertexCount = LodRenderData.StaticVertexBuffers.PositionVertexBuffer.GetNumVertices();
 
-	NeighborNodes.SetNum(0); // clear, to init clean
-	NeighborNodes.SetNum(VertexCount);
+	OutNeighborNodes.SetNum(0); // clear, to init clean
+	OutNeighborNodes.SetNum(VertexCount);
 
 	int32 BaseTriangle = 0;
 	int32 BaseVertex = 0;
@@ -96,7 +121,7 @@ BuildVertexToVertexAdjacencyBuffer(
 			for (int32 i = 0; i < 3; ++i)
 			{
 				const uint32 VertexIndex = RedirectionArray[V[i]] = UniqueIndexMap.FindOrAdd(P[i], V[i]);
-				TArray<uint32>& AdjacentVertices = NeighborNodes[VertexIndex];
+				TArray<uint32>& AdjacentVertices = OutNeighborNodes[VertexIndex];
 				for (int32 a = 1; a < 3; ++a)
 				{
 					const uint32 AdjacentVertexIndex = V[(i + a) % 3];
@@ -114,12 +139,12 @@ BuildVertexToVertexAdjacencyBuffer(
 			const int32 SiblingIndex = RedirectionArray[VertexIt];
 			if (SiblingIndex != VertexIt)
 			{
-				for (int32 i = 0; i < NeighborNodes[SiblingIndex].Num(); i++)
+				for (int32 i = 0; i < OutNeighborNodes[SiblingIndex].Num(); i++)
 				{
-					const uint32 OtherNode = NeighborNodes[SiblingIndex][i];
+					const uint32 OtherNode = OutNeighborNodes[SiblingIndex][i];
 					if (OtherNode != VertexIt)
 					{
-						NeighborNodes[VertexIt].AddUnique(OtherNode);
+						OutNeighborNodes[VertexIt].AddUnique(OtherNode);
 					}
 				}
 			}
@@ -167,9 +192,9 @@ FGenerateSurfaceBindings::Evaluate(UE::Dataflow::FContext& Context, const FDataf
 	auto DoubleVert = [](FVector3f V) { return FVector3d(V.X, V.Y, V.Z); };
 	auto TVec3Vert = [](FVector3f V) { return Chaos::TVec3<Chaos::FRealDouble>(V.X, V.Y, V.Z); };
 	
-	if (Out->IsA<FManagedArrayCollection>(&Collection))
+	if (Out->IsA(&Collection))
 	{
-		TUniquePtr<FTetrahedralCollection> InCollection(GetValue<FManagedArrayCollection>(Context, &Collection).NewCopy<FTetrahedralCollection>());
+		TUniquePtr<FTetrahedralCollection> InCollection(GetValue(Context, &Collection).NewCopy<FTetrahedralCollection>());
 		const TManagedArray<FIntVector4>* Tetrahedron = 
 			InCollection->FindAttribute<FIntVector4>(
 				FTetrahedralCollection::TetrahedronAttribute, FTetrahedralCollection::TetrahedralGroup);
@@ -197,22 +222,25 @@ FGenerateSurfaceBindings::Evaluate(UE::Dataflow::FContext& Context, const FDataf
 			InCollection->FindAttribute<FVector3f>(
 				"Vertex", "Vertices");
 
-		TObjectPtr<const USkeletalMesh> SkeletalMesh = GetValue<TObjectPtr<const USkeletalMesh>>(Context, &SkeletalMeshIn);
-		TObjectPtr<const UStaticMesh> StaticMesh = GetValue<TObjectPtr<const UStaticMesh>>(Context, &StaticMeshIn);
+		TObjectPtr<const USkeletalMesh> SkeletalMesh = GetValue(Context, &SkeletalMeshIn);
+		TObjectPtr<const UStaticMesh> StaticMesh = GetValue(Context, &StaticMeshIn);
 		TObjectPtr<UDynamicMesh> OutSKMDynamicMesh = NewObject<UDynamicMesh>();
 		OutSKMDynamicMesh->Reset();
 		FDynamicMesh3& OutSKMDynamicMesh3 = OutSKMDynamicMesh->GetMeshRef();
-		if ((SkeletalMesh || StaticMesh) && 
+
+		const bool UseSkeletalMesh = SkeletalMesh != nullptr;
+		const bool UseStaticMesh = StaticMesh != nullptr;
+		if ((UseSkeletalMesh || UseStaticMesh) &&
 			Tetrahedron && TetrahedronStart && TetrahedronCount &&
 			Triangle && FacesStart && FacesCount &&
 			Vertex)
 		{
 			// Extract positions to bind
 			FString MeshId;
-			TArray<TArray<FVector3f>> MeshVertices;
-			TArray<TArray<FIntVector>> SKMMeshTriangles;
-			TArray<TArray<TArray<uint32>>> MeshNeighborNodes;
-			if (SkeletalMesh)
+			TArray<TArray<FVector3f>> RenderMeshVertices;
+			TArray<TArray<FIntVector>> RenderMeshTriangles;
+			TArray<TArray<TArray<uint32>>> RenderMeshNeighborNodes;
+			if (UseSkeletalMesh)
 			{
 				FPrimaryAssetId Id = SkeletalMesh->GetPrimaryAssetId();
 				MeshId = ChaosFlesh::GetMeshId(SkeletalMesh, bUseSkeletalMeshImportModel);
@@ -221,16 +249,16 @@ FGenerateSurfaceBindings::Evaluate(UE::Dataflow::FContext& Context, const FDataf
 				{
 					FSkeletalMeshRenderData* RenderData = SkeletalMesh->GetResourceForRendering();
 
-					MeshVertices.SetNum(RenderData->LODRenderData.Num());
-					SKMMeshTriangles.SetNum(RenderData->LODRenderData.Num());
-					MeshNeighborNodes.SetNum(RenderData->LODRenderData.Num());
+					RenderMeshVertices.SetNum(RenderData->LODRenderData.Num());
+					RenderMeshTriangles.SetNum(RenderData->LODRenderData.Num());
+					RenderMeshNeighborNodes.SetNum(RenderData->LODRenderData.Num());
 					for (int32 i = 0; i < RenderData->LODRenderData.Num(); i++)
 					{
 						FSkeletalMeshLODRenderData* LODRenderData = &RenderData->LODRenderData[i];
 						const FPositionVertexBuffer& PositionVertexBuffer =
 							LODRenderData->StaticVertexBuffers.PositionVertexBuffer;
 
-						TArray<FVector3f>& Vertices = MeshVertices[i];
+						TArray<FVector3f>& Vertices = RenderMeshVertices[i];
 						Vertices.SetNumUninitialized(PositionVertexBuffer.GetNumVertices());
 						for (uint32 j = 0; j < PositionVertexBuffer.GetNumVertices(); j++)
 						{
@@ -238,18 +266,18 @@ FGenerateSurfaceBindings::Evaluate(UE::Dataflow::FContext& Context, const FDataf
 							Vertices[j] = Pos;
 						}
 
-						TArray<TArray<uint32>>& NeighborNodes = MeshNeighborNodes[i];
+						TArray<TArray<uint32>>& NeighborNodes = RenderMeshNeighborNodes[i];
 						BuildVertexToVertexAdjacencyBuffer(*LODRenderData, NeighborNodes);
-						BuildTriangles(*LODRenderData, SKMMeshTriangles[i]);
+						BuildTriangles(*LODRenderData, RenderMeshTriangles[i]);
 					}
 				}
 #if WITH_EDITOR
 				else // Import Model
 				{
 					const int32 LODIndex = 0;
-					MeshVertices.SetNum(1);
-					SKMMeshTriangles.SetNum(1);
-					MeshNeighborNodes.SetNum(1);
+					RenderMeshVertices.SetNum(1);
+					RenderMeshTriangles.SetNum(1);
+					RenderMeshNeighborNodes.SetNum(1);
 
 					// Check first if we have bulk data available and non-empty.
 					FMeshDescription SourceMesh;
@@ -270,47 +298,32 @@ FGenerateSurfaceBindings::Evaluate(UE::Dataflow::FContext& Context, const FDataf
 							SkeletalMeshModel->LODModels[LODIndex].GetMeshDescription(SkeletalMesh, LODIndex, SourceMesh);
 						}
 					}
-
-					FDynamicMesh3 DynamicMesh;
-					FMeshDescriptionToDynamicMesh Converter;
-					Converter.Convert(&SourceMesh, DynamicMesh);
-
-					TArray<FVector3f>& Vertices = MeshVertices[LODIndex];
-					Vertices.SetNumUninitialized(DynamicMesh.VertexCount());
-					for (int32 j = 0; j < DynamicMesh.VertexCount(); j++)
-					{
-						const FVector3d& Pos = DynamicMesh.GetVertex(j);
-						Vertices[j].Set(Pos[0], Pos[1], Pos[2]);
-					}
-
-					TArray<TArray<uint32>>& NeighborNodes = MeshNeighborNodes[LODIndex];
-					BuildVertexToVertexAdjacencyBuffer(DynamicMesh, NeighborNodes);
+					UnloadMeshDescription(SourceMesh, RenderMeshVertices[LODIndex], RenderMeshTriangles[LODIndex], RenderMeshNeighborNodes[LODIndex]);
 				}
 #endif
 			}
 			else // StaticMesh
 			{
 				MeshId = ChaosFlesh::GetMeshId(StaticMesh);
+				const int32 LODIndex = 0;
+				RenderMeshVertices.SetNum(1);
+				RenderMeshTriangles.SetNum(1);
+				RenderMeshNeighborNodes.SetNum(1);
+				FMeshDescription* MeshDescription = FGeometryCollectionEngineConversion::GetMaxResMeshDescriptionWithNormalsAndTangents(StaticMesh);
 
-				const FStaticMeshRenderData* RenderData = StaticMesh->GetRenderData();
-				const int32 NumLOD = RenderData->LODResources.Num();
-				MeshVertices.SetNum(NumLOD);
-				SKMMeshTriangles.SetNum(NumLOD);
-				for (int32 i = 0; i < NumLOD; i++)
+				if (MeshDescription)
 				{
-					const FStaticMeshLODResources& LODResources = RenderData->LODResources[i];
-					const FPositionVertexBuffer& PositionVertexBuffer = LODResources.VertexBuffers.PositionVertexBuffer;
-					
-					TArray<FVector3f>& Vertices = MeshVertices[i];
-					Vertices.SetNumUninitialized(PositionVertexBuffer.GetNumVertices());
-					for (uint32 j = 0; j < PositionVertexBuffer.GetNumVertices(); j++)
-					{
-						const FVector3f& Pos = PositionVertexBuffer.VertexPosition(j);
-						Vertices[j] = Pos;
-					}
+					UnloadMeshDescription(*MeshDescription, RenderMeshVertices[LODIndex], RenderMeshTriangles[LODIndex], RenderMeshNeighborNodes[LODIndex]);
+				}
+				else
+				{
+					Context.Warning(FString::Printf(
+						TEXT("No MeshDescription found in Static Mesh [%s]."),
+						*StaticMesh->GetName()),
+						this, Out);
 				}
 			}
-			TArray<FString> GeometryGroupGuidsLocal = GetValue<TArray<FString>>(Context, &GeometryGroupGuidsIn);
+			TArray<FString> GeometryGroupGuidsLocal = GetValue(Context, &GeometryGroupGuidsIn);
 			const TManagedArray<FString>* Guids = InCollection->FindAttribute<FString>("Guid", FGeometryCollection::GeometryGroup);
 
 			// Build Tetrahedra
@@ -332,7 +345,7 @@ FGenerateSurfaceBindings::Evaluate(UE::Dataflow::FContext& Context, const FDataf
 			TArray<int32> GeometryGroupSelected;
 			if (IsConnected(&TransformSelection))
 			{
-				FDataflowTransformSelection InTransformSelection = GetValue<FDataflowTransformSelection>(Context, &TransformSelection);
+				FDataflowTransformSelection InTransformSelection = GetValue(Context, &TransformSelection);
 				if (InTransformSelection.Num() == InCollection->NumElements(FGeometryCollection::TransformGroup))
 				{
 					GeometryGroupSelected = InCollection->TransformSelectionToGeometryIndices(InTransformSelection.AsArray());
@@ -441,33 +454,33 @@ FGenerateSurfaceBindings::Evaluate(UE::Dataflow::FContext& Context, const FDataf
 			// Do intersection tests against tets, then the surface.
 			//
 
-			TArray<TArray<FIntVector4>> Parents; Parents.SetNum(MeshVertices.Num());
-			TArray<TArray<FVector4f>> Weights;	 Weights.SetNum(MeshVertices.Num());
-			TArray<TArray<FVector3f>> Offsets;	 Offsets.SetNum(MeshVertices.Num());
-			TArray<TArray<float>> Masks;		 Masks.SetNum(MeshVertices.Num());
+			TArray<TArray<FIntVector4>> Parents; Parents.SetNum(RenderMeshVertices.Num());
+			TArray<TArray<FVector4f>> Weights;	 Weights.SetNum(RenderMeshVertices.Num());
+			TArray<TArray<FVector3f>> Offsets;	 Offsets.SetNum(RenderMeshVertices.Num());
+			TArray<TArray<float>> Masks;		 Masks.SetNum(RenderMeshVertices.Num());
 			TArray<int32> Orphans;
 			TArray<int32> OrphansCopy;
 			int32 TetHits = 0;
 			int32 TriHits = 0;
 			int32 Adoptions = 0;
 			int32 NumOrphans = 0;
-			for (int32 LOD = 0; LOD < MeshVertices.Num(); LOD++)
+			for (int32 LOD = 0; LOD < RenderMeshVertices.Num(); LOD++)
 			{
-				Parents[LOD].SetNumUninitialized(MeshVertices[LOD].Num());
-				Weights[LOD].SetNumUninitialized(MeshVertices[LOD].Num());
-				Offsets[LOD].SetNumUninitialized(MeshVertices[LOD].Num());
-				Masks[LOD].SetNumUninitialized(MeshVertices[LOD].Num());
+				Parents[LOD].SetNumUninitialized(RenderMeshVertices[LOD].Num());
+				Weights[LOD].SetNumUninitialized(RenderMeshVertices[LOD].Num());
+				Offsets[LOD].SetNumUninitialized(RenderMeshVertices[LOD].Num());
+				Masks[LOD].SetNumUninitialized(RenderMeshVertices[LOD].Num());
 				int32 NumTetNotCollocated = 0;
 
 				TArray<int32> TetIntersections; TetIntersections.Reserve(64);
-				for (int32 i = 0; i < MeshVertices[LOD].Num(); i++)
+				for (int32 i = 0; i < RenderMeshVertices[LOD].Num(); i++)
 				{
 					Parents[LOD][i] = FIntVector4(INDEX_NONE);
 					Weights[LOD][i] = FVector4f(0);
 					Offsets[LOD][i] = FVector3f(0);
 					Masks[LOD][i] = 0.0; // Shader does skinning for this vertex
 
-					const FVector3f& Pos = MeshVertices[LOD][i];
+					const FVector3f& Pos = RenderMeshVertices[LOD][i];
 					Chaos::TVec3<Chaos::FReal> PosD(Pos[0], Pos[1], Pos[2]);
 					TetIntersections = TetBVH.FindAllIntersections(PosD);
 					int32 j = 0;
@@ -574,11 +587,11 @@ FGenerateSurfaceBindings::Evaluate(UE::Dataflow::FContext& Context, const FDataf
 				// 
 				// Advancing front orphan reparenting
 				//
-				if (!MeshNeighborNodes.IsValidIndex(LOD))
+				if (!RenderMeshNeighborNodes.IsValidIndex(LOD))
 				{
 					continue;
 				}
-				const TArray<TArray<uint32>>& NeighborNodes = MeshNeighborNodes[LOD];
+				const TArray<TArray<uint32>>& NeighborNodes = RenderMeshNeighborNodes[LOD];
 				OrphansCopy = Orphans;
 				TSet<int32> OrphanSet(Orphans);
 				while (bDoOrphanReparenting && Orphans.Num())
@@ -621,7 +634,7 @@ FGenerateSurfaceBindings::Evaluate(UE::Dataflow::FContext& Context, const FDataf
 						// We only have orphans with no neighbors left.
 						break;
 					}
-					const FVector3f& Pos = MeshVertices[LOD][Orphan];
+					const FVector3f& Pos = RenderMeshVertices[LOD][Orphan];
 					Chaos::TVec3<Chaos::FReal> PosD(Pos[0], Pos[1], Pos[2]);
 
 					// Use the parent simplices of non-orphan neighbors as test candidates.
@@ -698,18 +711,34 @@ FGenerateSurfaceBindings::Evaluate(UE::Dataflow::FContext& Context, const FDataf
 				} // end while(Orphans)
 				NumOrphans += Orphans.Num();
 
-				//ELogVerbosity::Type Verbosity = Orphans.Num() > 0 ? ELogVerbosity::Error : ELogVerbosity::Display;
-				UE_LOG(LogMeshBindings, Display,
-					TEXT("'%s' - Generated mesh bindings between tet mesh and %s mesh of '%s' LOD %d - stats:\n"
-						"    Render vertices num: %d\n"
-						"    Vertices in tetrahedra: %d\n"
-						"    Vertices bound to tet surface: %d\n"
-						"    Orphaned vertices reparented: %d\n"
-						"    Vertices orphaned: %d"),
-					*GetName().ToString(),
-					bUseSkeletalMeshImportModel ? TEXT("import") : TEXT("render"),
-					*MeshId, LOD,
-					MeshVertices[LOD].Num(), TetHits, TriHits, Adoptions, NumOrphans);
+				if (Orphans.Num() > 0)
+				{
+					UE_LOG(LogMeshBindings, Error,
+						TEXT("'%s' - Generated mesh bindings between tet mesh and %s mesh of '%s' LOD %d - stats:\n"
+							"    Render vertices num: %d\n"
+							"    Vertices in tetrahedra: %d\n"
+							"    Vertices bound to tet surface: %d\n"
+							"    Orphaned vertices reparented: %d\n"
+							"    Vertices orphaned: %d"),
+						*GetName().ToString(),
+						bUseSkeletalMeshImportModel ? TEXT("import") : TEXT("render"),
+						*MeshId, LOD,
+						RenderMeshVertices[LOD].Num(), TetHits, TriHits, Adoptions, NumOrphans);
+				}
+				else
+				{
+					UE_LOG(LogMeshBindings, Display,
+						TEXT("'%s' - Generated mesh bindings between tet mesh and %s mesh of '%s' LOD %d - stats:\n"
+							"    Render vertices num: %d\n"
+							"    Vertices in tetrahedra: %d\n"
+							"    Vertices bound to tet surface: %d\n"
+							"    Orphaned vertices reparented: %d\n"
+							"    Vertices orphaned: %d"),
+						*GetName().ToString(),
+						bUseSkeletalMeshImportModel ? TEXT("import") : TEXT("render"),
+						* MeshId, LOD,
+						RenderMeshVertices[LOD].Num(), TetHits, TriHits, Adoptions, NumOrphans);
+				}
 				if (NumTetNotCollocated)
 				{
 					UE_LOG(LogMeshBindings, Error,
@@ -724,27 +753,26 @@ FGenerateSurfaceBindings::Evaluate(UE::Dataflow::FContext& Context, const FDataf
 			GeometryCollection::Facades::FTetrahedralBindings TetBindings(*InCollection);
 			TetBindings.DefineSchema();
 			FName MeshName(*MeshId, MeshId.Len());
-			for (int32 LOD = 0; LOD < MeshVertices.Num(); LOD++)
+			for (int32 LOD = 0; LOD < RenderMeshVertices.Num(); LOD++)
 			{
 				TetBindings.AddBindingsGroup(/*TetMeshIdx = */ 0, MeshName, LOD);
 				TetBindings.SetBindingsData(Parents[LOD], Weights[LOD], Offsets[LOD], Masks[LOD]);
 			}
 
 			//Write DynamicMesh			
-			if (MeshVertices.Num())
+			if (RenderMeshVertices.Num())
 			{
 				OutSKMDynamicMesh3.EnableAttributes();
 				OutSKMDynamicMesh3.EnableVertexColors(FVector3f(1, 0, 0));
-				TBitArray<> WasOrphan(false, MeshVertices[0].Num());
-				UE_LOG(LogTemp, Warning, TEXT("Num of previous orphans: %d"), OrphansCopy.Num());
+				TBitArray<> WasOrphan(false, RenderMeshVertices[0].Num());
 				for (int32 OrphanIdx : OrphansCopy)
 				{
 					WasOrphan[OrphanIdx] = true;
 				}
-				for (int32 VertexIndex = 0; VertexIndex < MeshVertices[0].Num(); ++VertexIndex)
+				for (int32 VertexIndex = 0; VertexIndex < RenderMeshVertices[0].Num(); ++VertexIndex)
 				{
 					FVertexInfo VertexInfo;
-					VertexInfo.Position = DoubleVert(MeshVertices[0][VertexIndex]);
+					VertexInfo.Position = DoubleVert(RenderMeshVertices[0][VertexIndex]);
 					VertexInfo.bHaveC = true;
 					if (Parents[0][VertexIndex][0] == INDEX_NONE)
 					{
@@ -768,19 +796,19 @@ FGenerateSurfaceBindings::Evaluate(UE::Dataflow::FContext& Context, const FDataf
 					OutSKMDynamicMesh3.AppendVertex(VertexInfo);
 				}
 				
-				for (int32 TriangleIndex = 0; TriangleIndex < SKMMeshTriangles[0].Num(); ++TriangleIndex)
+				for (int32 TriangleIndex = 0; TriangleIndex < RenderMeshTriangles[0].Num(); ++TriangleIndex)
 				{
 					OutSKMDynamicMesh3.AppendTriangle(FIndex3i(
-						SKMMeshTriangles[0][TriangleIndex][0],
-						SKMMeshTriangles[0][TriangleIndex][1],
-						SKMMeshTriangles[0][TriangleIndex][2])
+						RenderMeshTriangles[0][TriangleIndex][0],
+						RenderMeshTriangles[0][TriangleIndex][1],
+						RenderMeshTriangles[0][TriangleIndex][2])
 					);
 				}
 				// Compute normals
 				OutSKMDynamicMesh3.EnableVertexNormals(FVector3f(1, 0, 0));
 				FMeshNormals MeshNormals(&OutSKMDynamicMesh3);
 				MeshNormals.ComputeVertexNormals();
-				for (int32 VertexIndex = 0; VertexIndex < MeshVertices[0].Num(); ++VertexIndex)
+				for (int32 VertexIndex = 0; VertexIndex < RenderMeshVertices[0].Num(); ++VertexIndex)
 				{
 					OutSKMDynamicMesh3.SetVertexNormal(VertexIndex, FloatVert(MeshNormals[VertexIndex]));
 				}
diff --git a/Engine/Plugins/Experimental/ChaosFlesh/Source/ChaosFleshNodes/Public/Dataflow/ChaosFleshGenerateSurfaceBindingsNode.h b/Engine/Plugins/Experimental/ChaosFlesh/Source/ChaosFleshNodes/Public/Dataflow/ChaosFleshGenerateSurfaceBindingsNode.h
--- a/Engine/Plugins/Experimental/ChaosFlesh/Source/ChaosFleshNodes/Public/Dataflow/ChaosFleshGenerateSurfaceBindingsNode.h
+++ b/Engine/Plugins/Experimental/ChaosFlesh/Source/ChaosFleshNodes/Public/Dataflow/ChaosFleshGenerateSurfaceBindingsNode.h
@@ -38,8 +38,12 @@ public:
 		RegisterOutputConnection(&Collection, &Collection);
 		RegisterInputConnection(&StaticMeshIn);
 		RegisterInputConnection(&SkeletalMeshIn);
-		RegisterInputConnection(&TransformSelection);
-		RegisterInputConnection(&GeometryGroupGuidsIn);
+		RegisterInputConnection(&TransformSelection)
+			.SetCanHidePin(true)
+			.SetPinIsHidden(true);
+		RegisterInputConnection(&GeometryGroupGuidsIn)
+			.SetCanHidePin(true)
+			.SetPinIsHidden(true);
 		RegisterOutputConnection(&SKMDynamicMesh)
 			.SetCanHidePin(true)
 			.SetPinIsHidden(true);
@@ -51,11 +55,11 @@ private:
 	FManagedArrayCollection Collection;
 
 	/** The input mesh, whose render surface is used to generate bindings. */
-	UPROPERTY(EditAnywhere, Category = "Dataflow", meta = (DataflowInput, DisplayName = "StaticMesh"))
+	UPROPERTY(EditAnywhere, Category = "Dataflow", DisplayName = "Static Mesh", meta = (DataflowInput, EditCondition = "SkeletalMeshIn == nullptr"))
 	TObjectPtr<const UStaticMesh> StaticMeshIn = nullptr;
 
 	/** The input mesh, whose render surface is used to generate bindings. */
-	UPROPERTY(EditAnywhere, Category = "Dataflow", meta = (DataflowInput, DisplayName = "SkeletalMesh"))
+	UPROPERTY(EditAnywhere, Category = "Dataflow", DisplayName = "Skeletal Mesh", meta = (DataflowInput, EditCondition = "StaticMeshIn == nullptr"))
 	TObjectPtr<const USkeletalMesh> SkeletalMeshIn = nullptr;
 
 	/** Only meshes with transforms in TransformSelection will be bound. */
@@ -67,7 +71,7 @@ private:
 	TArray<FString> GeometryGroupGuidsIn;
 
 	/** Use the import geometry of the skeletal mesh. */
-	UPROPERTY(EditAnywhere, Category = "Dataflow", meta = (DisplayName = "UseSkeletalMeshImportModel"))
+	UPROPERTY(EditAnywhere, Category = "Dataflow", meta = (DisplayName = "UseSkeletalMeshImportModel", EditCondition = "MeshInputType == EChaosFleshGenerateSurfaceBindingsMeshInputType::SkeletalMesh", EditConditionHides))
 	bool bUseSkeletalMeshImportModel = false;
 
 	/** Enable binding to the exterior hull of the tetrahedron mesh. */
diff --git a/Engine/Plugins/Experimental/Dataflow/Source/DataflowEditor/Private/Dataflow/DataflowSimulationGeometryCache.cpp b/Engine/Plugins/Experimental/Dataflow/Source/DataflowEditor/Private/Dataflow/DataflowSimulationGeometryCache.cpp
--- a/Engine/Plugins/Experimental/Dataflow/Source/DataflowEditor/Private/Dataflow/DataflowSimulationGeometryCache.cpp
+++ b/Engine/Plugins/Experimental/Dataflow/Source/DataflowEditor/Private/Dataflow/DataflowSimulationGeometryCache.cpp
@@ -8,6 +8,7 @@
 #include "Rendering/SkeletalMeshRenderData.h"
 #include "Rendering/SkeletalMeshModel.h"
 #include "Engine/SkeletalMesh.h"
+#include "Engine/StaticMesh.h"
 
 #if WITH_EDITOR 
 #include "FileHelpers.h"
@@ -73,6 +74,24 @@ namespace UE::DataflowSimulationGeometryCache
 #endif
 	}
 
+	void SaveGeometryCache(UGeometryCache& GeometryCache, float FrameRate, const UStaticMesh& StaticMesh, TArrayView<TArray<FVector3f>> PositionsToMoveFrom)
+	{
+#if WITH_EDITOR 
+		using UE::GeometryCacheHelpers::FGeometryCacheConstantTopologyWriter;
+		using UE::GeometryCacheHelpers::AddTrackWriterFromStaticMesh;
+		FGeometryCacheConstantTopologyWriter::FConfig Config; Config.FPS = FrameRate;
+		using FTrackWriter = FGeometryCacheConstantTopologyWriter::FTrackWriter;
+		FGeometryCacheConstantTopologyWriter Writer(GeometryCache, Config);
+		const int32 Index = AddTrackWriterFromStaticMesh(Writer, StaticMesh);
+		if (Index == INDEX_NONE)
+		{
+			return;
+		}
+		FTrackWriter& TrackWriter = Writer.GetTrackWriter(Index);
+		TrackWriter.WriteAndClose(PositionsToMoveFrom);
+#endif
+	}
+
 	void SavePackage(UObject& Object)
 	{
 #if WITH_EDITOR 
diff --git a/Engine/Plugins/Experimental/Dataflow/Source/DataflowEditor/Private/Dataflow/DataflowSimulationScene.cpp b/Engine/Plugins/Experimental/Dataflow/Source/DataflowEditor/Private/Dataflow/DataflowSimulationScene.cpp
--- a/Engine/Plugins/Experimental/Dataflow/Source/DataflowEditor/Private/Dataflow/DataflowSimulationScene.cpp
+++ b/Engine/Plugins/Experimental/Dataflow/Source/DataflowEditor/Private/Dataflow/DataflowSimulationScene.cpp
@@ -59,6 +59,7 @@ FDataflowSimulationScene::FDataflowSimulationScene(FPreviewScene::ConstructionVa
 			SceneDescription->BlueprintTransform = DataflowAsset->PreviewBlueprintTransform; 
 			SceneDescription->GeometryCacheAsset = Cast<UGeometryCache>(DataflowAsset->PreviewGeometryCacheAsset.LoadSynchronous());
 			SceneDescription->EmbeddedSkeletalMesh = Cast<USkeletalMesh>(DataflowAsset->PreviewEmbeddedSkeletalMesh.LoadSynchronous());
+			SceneDescription->EmbeddedStaticMesh = Cast<UStaticMesh>(DataflowAsset->PreviewEmbeddedStaticMesh.LoadSynchronous());
 		}
 		if(SceneDescription->BlueprintClass == nullptr)
 		{
@@ -447,7 +448,7 @@ void FDataflowSimulationScene::SceneDescriptionPropertyChanged(const FName& Prop
 			DataflowAsset->PreviewBlueprintTransform = SceneDescription->BlueprintTransform;
 			DataflowAsset->PreviewGeometryCacheAsset = SceneDescription->GeometryCacheAsset;
 			DataflowAsset->PreviewEmbeddedSkeletalMesh = SceneDescription->EmbeddedSkeletalMesh;
-
+			DataflowAsset->PreviewEmbeddedStaticMesh = SceneDescription->EmbeddedStaticMesh;
 			DataflowAsset->MarkPackageDirty();
 #endif
 		}
@@ -469,7 +470,9 @@ void UDataflowSimulationSceneDescription::GenerateGeometryCache()
 	float Time = TimeRange[0];
 	TObjectPtr<AActor> GetRootActor = SimulationScene->GetRootActor();
 	TObjectPtr<AActor> PreviewActor = SimulationScene->GetPreviewActor();
-	if (CacheAsset && GeometryCacheAsset && GetRootActor && EmbeddedSkeletalMesh)
+	const bool UseSkeletalMesh = EmbeddedSkeletalMesh != nullptr;
+	const bool UseStaticMesh = EmbeddedStaticMesh != nullptr;
+	if (CacheAsset && GeometryCacheAsset && GetRootActor && (UseSkeletalMesh || UseStaticMesh))
 	{
 		IDataflowGeometryCachable* GeometryCachable = nullptr; //interface for ChaosDeformableTetrahedralComponent
 
@@ -489,20 +492,36 @@ void UDataflowSimulationSceneDescription::GenerateGeometryCache()
 			UE_LOG(LogDataflowSimulationGeometryCache, Error, TEXT("No GeometryCachable Component in the Preview Actor"));
 			return;
 		}
-		TOptional<TArray<int32>> OptionalMap = GeometryCachable->GetMeshImportVertexMap(*EmbeddedSkeletalMesh);
-		if (!OptionalMap)
-		{
-			return;
-		}
-		const TArray<int32>& Map = OptionalMap.GetValue();
-		TArray<uint32> ImportedVertexNumbers = TArray<uint32>(reinterpret_cast<const uint32*>(Map.GetData()), Map.Num());
 		for (int32 Frame = 0; Frame < NumFrames; ++Frame)
 		{
 			Cast<AChaosCacheManager>(GetRootActor)->SetStartTime(Time);
-			RenderPositions[Frame] = GeometryCachable->GetGeometryCachePositions(EmbeddedSkeletalMesh);
+			if (UseSkeletalMesh)
+			{
+				RenderPositions[Frame] = GeometryCachable->GetGeometryCachePositions(EmbeddedSkeletalMesh);
+			}
+			else
+			{
+				RenderPositions[Frame] = GeometryCachable->GetGeometryCachePositions(EmbeddedStaticMesh);
+			}
 			Time += SimulationScene->GetDeltaTime();
 		}
-		UE::DataflowSimulationGeometryCache::SaveGeometryCache(*GeometryCacheAsset, float(CacheParams.FrameRate), *EmbeddedSkeletalMesh, ImportedVertexNumbers, RenderPositions);
+		if (UseSkeletalMesh)
+		{
+			TOptional<TArray<int32>> OptionalMap = GeometryCachable->GetMeshImportVertexMap(*EmbeddedSkeletalMesh);
+			if (!OptionalMap)
+			{
+				UE_LOG(LogDataflowSimulationGeometryCache, Error, TEXT("Failed to get MeshImportVertexMap for the skeletal mesh. See the log for more info."));
+				return;
+			}
+			const TArray<int32>& Map = OptionalMap.GetValue();
+			TArray<uint32> ImportedVertexNumbers = TArray<uint32>(reinterpret_cast<const uint32*>(Map.GetData()), Map.Num());
+			UE::DataflowSimulationGeometryCache::SaveGeometryCache(*GeometryCacheAsset, float(CacheParams.FrameRate), *EmbeddedSkeletalMesh, ImportedVertexNumbers, RenderPositions);
+		}
+		else // UseStaticMesh
+		{
+			UE::DataflowSimulationGeometryCache::SaveGeometryCache(*GeometryCacheAsset, float(CacheParams.FrameRate), *EmbeddedStaticMesh, RenderPositions);
+		}
+
 		UE::DataflowSimulationGeometryCache::SavePackage(*GeometryCacheAsset);
 	}
 }
diff --git a/Engine/Plugins/Experimental/Dataflow/Source/DataflowEditor/Public/Dataflow/DataflowSimulationGeometryCache.h b/Engine/Plugins/Experimental/Dataflow/Source/DataflowEditor/Public/Dataflow/DataflowSimulationGeometryCache.h
--- a/Engine/Plugins/Experimental/Dataflow/Source/DataflowEditor/Public/Dataflow/DataflowSimulationGeometryCache.h
+++ b/Engine/Plugins/Experimental/Dataflow/Source/DataflowEditor/Public/Dataflow/DataflowSimulationGeometryCache.h
@@ -7,15 +7,19 @@
 #include "GeometryCollection/ManagedArray.h"
 
 class USkinnedAsset;
+class UStaticMesh;
 class UGeometryCache;
 
 DECLARE_LOG_CATEGORY_EXTERN(LogDataflowSimulationGeometryCache, Log, All);
 
 namespace UE::DataflowSimulationGeometryCache
 {
-	/** Saves positions into the geometry cache */
+	/** Saves skinned asset positions into the geometry cache*/
 	DATAFLOWEDITOR_API void SaveGeometryCache(UGeometryCache& GeometryCache, float FrameRate, const USkinnedAsset& Asset, TConstArrayView<uint32> ImportedVertexNumbers, TArrayView<TArray<FVector3f>> PositionsToMoveFrom);
 
+	/** Saves static mesh positions into the geometry cache */
+	DATAFLOWEDITOR_API void SaveGeometryCache(UGeometryCache& GeometryCache, float FrameRate, const UStaticMesh& StaticMesh, TArrayView<TArray<FVector3f>> PositionsToMoveFrom);
+
 	/** Saves the geometry cache */
 	DATAFLOWEDITOR_API void SavePackage(UObject& Object);
 };
\ No newline at end of file
diff --git a/Engine/Plugins/Experimental/Dataflow/Source/DataflowEditor/Public/Dataflow/DataflowSimulationScene.h b/Engine/Plugins/Experimental/Dataflow/Source/DataflowEditor/Public/Dataflow/DataflowSimulationScene.h
--- a/Engine/Plugins/Experimental/Dataflow/Source/DataflowEditor/Public/Dataflow/DataflowSimulationScene.h
+++ b/Engine/Plugins/Experimental/Dataflow/Source/DataflowEditor/Public/Dataflow/DataflowSimulationScene.h
@@ -53,16 +53,20 @@ public:
 	UPROPERTY(EditAnywhere, Category = "Geometry Cache", DisplayName="Geometry Cache", meta=(EditCondition = "CacheAsset != nullptr"))
 	TObjectPtr<UGeometryCache> GeometryCacheAsset = nullptr;
 
-	/** SkeletalMesh interpolated from simulation. This should match the SkeletalMesh used in GenerateSurfaceBindings node */
-	UPROPERTY(EditAnywhere, Category = "Geometry Cache", DisplayName = "EmbeddedSkeletalMesh", meta = (EditCondition = "CacheAsset != nullptr"))
+	/** Skeletal mesh interpolated from simulation. This should match the SkeletalMesh used in GenerateSurfaceBindings node */
+	UPROPERTY(EditAnywhere, Category = "Geometry Cache", meta = (EditCondition = "CacheAsset != nullptr && EmbeddedStaticMesh == nullptr"))
 	TObjectPtr<USkeletalMesh> EmbeddedSkeletalMesh = nullptr;
 
+	/** Static mesh interpolated from simulation. This should match the Static mesh used in GenerateSurfaceBindings node */
+	UPROPERTY(EditAnywhere, Category = "Geometry Cache", meta = (EditCondition = "CacheAsset != nullptr && EmbeddedSkeletalMesh == nullptr"))
+	TObjectPtr<UStaticMesh> EmbeddedStaticMesh = nullptr;
+
 	/** Interpolates and saves geometry cache from Chaos cache */
-	UFUNCTION(CallInEditor, Category = "Geometry Cache", meta = (EditCondition = "CacheAsset != nullptr && GeometryCacheAsset != nullptr && EmbeddedSkeletalMesh != nullptr"))
+	UFUNCTION(CallInEditor, Category = "Geometry Cache")
 	void GenerateGeometryCache();
 
 	/** Creates a new geometry cache file */
-	UFUNCTION(CallInEditor, Category = "Geometry Cache", meta = (EditCondition = "CacheAsset != nullptr && GeometryCacheAsset != nullptr && EmbeddedSkeletalMesh != nullptr"))
+	UFUNCTION(CallInEditor, Category = "Geometry Cache")
 	void NewGeometryCache();
 
 	/** Visibility of the skeletal mesh */
diff --git a/Engine/Plugins/Runtime/GeometryCache/Source/GeometryCache/GeometryCache.Build.cs b/Engine/Plugins/Runtime/GeometryCache/Source/GeometryCache/GeometryCache.Build.cs
--- a/Engine/Plugins/Runtime/GeometryCache/Source/GeometryCache/GeometryCache.Build.cs
+++ b/Engine/Plugins/Runtime/GeometryCache/Source/GeometryCache/GeometryCache.Build.cs
@@ -20,9 +20,15 @@ public class GeometryCache : ModuleRules
                 "NiagaraCore"
 			}
 		);
-		PrivateDependencyModuleNames.Add("MeshDescription");
-		PrivateDependencyModuleNames.Add("StaticMeshDescription");
-		PrivateDependencyModuleNames.Add("TraceLog");
+		PrivateDependencyModuleNames.AddRange(
+			new string[] {
+				"GeometryCore",
+				"MeshConversion",
+				"MeshDescription",
+				"StaticMeshDescription",
+				"TraceLog",
+			}
+		);
 
 		PublicIncludePathModuleNames.Add("TargetPlatform");
 		PrivateIncludePathModuleNames.Add("Shaders");
diff --git a/Engine/Plugins/Runtime/GeometryCache/Source/GeometryCache/Private/GeometryCacheConstantTopologyWriter.cpp b/Engine/Plugins/Runtime/GeometryCache/Source/GeometryCache/Private/GeometryCacheConstantTopologyWriter.cpp
--- a/Engine/Plugins/Runtime/GeometryCache/Source/GeometryCache/Private/GeometryCacheConstantTopologyWriter.cpp
+++ b/Engine/Plugins/Runtime/GeometryCache/Source/GeometryCache/Private/GeometryCacheConstantTopologyWriter.cpp
@@ -4,12 +4,15 @@
 
 #include "Engine/SkinnedAsset.h"
 #include "Engine/SkinnedAssetCommon.h"
+#include "Engine/StaticMesh.h"
 #include "GeometryCache.h"
 #include "GeometryCacheCodecV1.h"
 #include "GeometryCacheMeshData.h"
 #include "GeometryCacheTrackStreamable.h"
 #include "Rendering/SkeletalMeshRenderData.h"
 #include "Rendering/SkeletalMeshModel.h"
+#include "MeshDescription.h"
+#include "MeshDescriptionToDynamicMesh.h"
 
 DEFINE_LOG_CATEGORY_STATIC(LogGeometryCacheWriter, Log, All);
 
@@ -517,5 +520,77 @@ namespace UE::GeometryCacheHelpers
 		return NumAddedTracks;
 	}
 
+	int32 AddTrackWriterFromStaticMesh(FGeometryCacheConstantTopologyWriter& Writer, const UStaticMesh& Asset)
+	{
+		constexpr int32 LODIndex = 0;
+		return AddTrackWriterFromStaticMeshAndMaterials(Writer, Asset, LODIndex, {});
+	}
+
+	int32 AddTrackWriterFromStaticMeshAndMaterials(
+		FGeometryCacheConstantTopologyWriter& Writer,
+		const UStaticMesh& StaticMesh,
+		int32 LODIndex,
+		const TArray<TObjectPtr<UMaterialInterface>>& Materials)
+	{
+		FMeshDescription* MeshDescription = nullptr;
+#if WITH_EDITORONLY_DATA
+		MeshDescription = StaticMesh.GetMeshDescription(LODIndex);
+#endif //WITH_EDITORONLY_DATA
+		if (!MeshDescription)
+		{
+			return INDEX_NONE;
+		}
+		
+		FMeshDescriptionToDynamicMesh Converter;
+		UE::Geometry::FDynamicMesh3 DynamicMesh;
+		Converter.Convert(MeshDescription, DynamicMesh);
+		const int VertexCount = DynamicMesh.VertexCount();
+		const int TriangleCount = DynamicMesh.TriangleCount();
+		// Write data
+		FTrackWriter& TrackWriter = Writer.AddTrackWriter();
+		TArray<uint32> Indices;
+		TArray<FVector2f> UVs;
+		TArray<FColor> Colors;
+
+		Indices.SetNumUninitialized(3 * TriangleCount);
+		for (int32 TriIdx = 0; TriIdx < TriangleCount; ++TriIdx)
+		{
+			const UE::Geometry::FIndex3i& Tri = DynamicMesh.GetTriangle(TriIdx);
+			for (int32 CornerIdx = 0; CornerIdx < 3; ++CornerIdx)
+			{
+				Indices[3 * TriIdx + CornerIdx] = Tri[CornerIdx];
+			}
+		}
+		
+		TrackWriter.Indices = MoveTemp(Indices);
+		UVs.SetNumUninitialized(VertexCount);
+		Colors.SetNumUninitialized(VertexCount);
+		TrackWriter.ImportedVertexNumbers.SetNum(VertexCount);
+		for (int VertexIdx = 0; VertexIdx < VertexCount; ++VertexIdx)
+		{
+			TrackWriter.ImportedVertexNumbers[VertexIdx] = VertexIdx;
+			UVs[VertexIdx] = DynamicMesh.GetVertexUV(VertexIdx);
+			Colors[VertexIdx] = FLinearColor(DynamicMesh.GetVertexColor(VertexIdx)).ToFColor(false);
+		}
+		TrackWriter.UVs = MoveTemp(UVs);
+		TrackWriter.Colors = MoveTemp(Colors);
+
+		const int32 MaterialOffset = Writer.GetNumMaterials();
+		const TArray<FStaticMaterial>& AssetMaterials = StaticMesh.GetStaticMaterials();
+		TArray<TObjectPtr<UMaterialInterface>> MaterialPtrs;
+		for (const FStaticMaterial& Mat : AssetMaterials)
+		{
+			MaterialPtrs.Add(Mat.MaterialInterface);
+		}
+		Writer.AddMaterials(MaterialPtrs);
+		// Only one batch
+		FGeometryCacheMeshBatchInfo BatchInfo;
+		BatchInfo.StartIndex = 0;
+		BatchInfo.NumTriangles = TriangleCount;
+		BatchInfo.MaterialIndex = MaterialOffset;
+		TrackWriter.BatchesInfo.Add(BatchInfo);
+		return Writer.GetNumTracks() - 1;
+	}
+
 #endif // WITH_EDITOR
 };
\ No newline at end of file
diff --git a/Engine/Plugins/Runtime/GeometryCache/Source/GeometryCache/Public/GeometryCacheConstantTopologyWriter.h b/Engine/Plugins/Runtime/GeometryCache/Source/GeometryCache/Public/GeometryCacheConstantTopologyWriter.h
--- a/Engine/Plugins/Runtime/GeometryCache/Source/GeometryCache/Public/GeometryCacheConstantTopologyWriter.h
+++ b/Engine/Plugins/Runtime/GeometryCache/Source/GeometryCache/Public/GeometryCacheConstantTopologyWriter.h
@@ -13,6 +13,7 @@ class UGeometryCache;
 class UGeometryCacheTrack;
 class UMaterialInterface; 
 class USkinnedAsset;
+class UStaticMesh;
 
 namespace UE::GeometryCacheHelpers
 {
@@ -141,6 +142,30 @@ namespace UE::GeometryCacheHelpers
 	 */
 	GEOMETRYCACHE_API int32 AddTrackWriterFromSkinnedAssetAndMaterials(FGeometryCacheConstantTopologyWriter& Writer, const USkinnedAsset& Asset, int32 LODIndex, const TArray<TObjectPtr<UMaterialInterface>>& Materials);
 
+	/**
+	 * @brief This will create a track writer and fill in the track writer's data (indices, UVs, materials .etc) from the static mesh.
+	 * This only creates one track for the static mesh. If you want to create multiple tracks based on imported mesh info, checkout AddTrackWriterFromSkinnedAsset() in ChaosClothGenerator.cpp.
+	 * Usage:
+	 * @code
+	 * 	FGeometryCacheConstantTopologyWriter Writer(MyCache);
+	 * 	int32 Index = AddTrackWriterFromStaticMesh(Writer, Asset);
+	 * 	if (Index != INDEX_NONE)
+	 * 	{
+	 * 		Writer.GetTrackWriter(Index).WriteAndClose(PositionsToMoveFrom);
+	 * 	}
+	 * @endcode
+	 */
+	GEOMETRYCACHE_API int32 AddTrackWriterFromStaticMesh(FGeometryCacheConstantTopologyWriter& Writer, const UStaticMesh& Asset);
+
+	/**
+	 * @brief This will create a track writer and fill in the track writer's data (indices, UVs, .etc) from the static mesh, with the option to specify a set of materials
+	 */
+	GEOMETRYCACHE_API int32 AddTrackWriterFromStaticMeshAndMaterials(
+		FGeometryCacheConstantTopologyWriter& Writer,
+		const UStaticMesh& StaticMesh,
+		int32 LODIndex,
+		const TArray<TObjectPtr<UMaterialInterface>>& Materials);
+
 	/**
 	 * @brief This will create multiple track writers and fill in the track writer's data (indices, UVs, materials .etc) from the template geometry cache. 
 	 * The number of track writers created equals to the number of tracks in the template geometry cache. 
diff --git a/Engine/Source/Runtime/Experimental/Chaos/Public/GeometryCollection/Facades/CollectionTetrahedralBindingsFacade.h b/Engine/Source/Runtime/Experimental/Chaos/Public/GeometryCollection/Facades/CollectionTetrahedralBindingsFacade.h
--- a/Engine/Source/Runtime/Experimental/Chaos/Public/GeometryCollection/Facades/CollectionTetrahedralBindingsFacade.h
+++ b/Engine/Source/Runtime/Experimental/Chaos/Public/GeometryCollection/Facades/CollectionTetrahedralBindingsFacade.h
@@ -216,20 +216,45 @@ namespace GeometryCollection::Facades
 				{
 					for (const FIntVector4& Tet : ParentsArray->Get())
 					{
-						const int32 iEnd = Tet[3] == INDEX_NONE ? 3 : 4;
-						for (int32 i = 0; i < iEnd; i++) MinIndexValue = FMath::Min(MinIndexValue, Tet[i]);
+						const int32 IndexEnd = Tet[3] == INDEX_NONE ? 3 : 4;
+						for (int32 LocalIdx = 0; LocalIdx < IndexEnd; ++LocalIdx)
+						{
+							MinIndexValue = FMath::Min(MinIndexValue, Tet[LocalIdx]);
+						}
 					}
 				}
 				return MinIndexValue;
 			}
+			CHAOS_API int32 MinIndexPosition() const
+			{
+				MinIndexValue = TNumericLimits<int32>::Max();
+				int32 OutMinIndexPosition = INDEX_NONE;
+				for (int32 TetIdx = 0; TetIdx < ParentsArray->Get().Num(); ++TetIdx)
+				{
+					const FIntVector4& Tet = ParentsArray->Get()[TetIdx];
+					const int32 IndexEnd = Tet[3] == INDEX_NONE ? 3 : 4;
+					for (int32 LocalIdx = 0; LocalIdx < IndexEnd; ++LocalIdx)
+					{
+						if (Tet[LocalIdx] < MinIndexValue)
+						{
+							MinIndexValue = FMath::Min(MinIndexValue, Tet[LocalIdx]);
+							OutMinIndexPosition = TetIdx;
+						}
+					}
+				}
+				return OutMinIndexPosition;
+			}
 			CHAOS_API int32 MaxIndex() const
 			{
 				if (MaxIndexValue == INDEX_NONE)
 				{
 					for (const FIntVector4& Tet : ParentsArray->Get())
 					{
-						const int32 iEnd = Tet[3] == INDEX_NONE ? 3 : 4;
-						for (int32 i = 0; i < iEnd; i++) MaxIndexValue = FMath::Max(MaxIndexValue, Tet[i]);
+						const int32 IndexEnd = Tet[3] == INDEX_NONE ? 3 : 4;
+						for (int32 LocalIdx = 0; LocalIdx < IndexEnd; ++LocalIdx)
+						{
+							MaxIndexValue = FMath::Max(MaxIndexValue, Tet[LocalIdx]);
+						}
 					}
 				}
 				return MaxIndexValue;
diff --git a/Engine/Source/Runtime/Experimental/Dataflow/Engine/Public/Dataflow/DataflowObject.h b/Engine/Source/Runtime/Experimental/Dataflow/Engine/Public/Dataflow/DataflowObject.h
--- a/Engine/Source/Runtime/Experimental/Dataflow/Engine/Public/Dataflow/DataflowObject.h
+++ b/Engine/Source/Runtime/Experimental/Dataflow/Engine/Public/Dataflow/DataflowObject.h
@@ -250,6 +250,10 @@ public:
 	/** SkeletalMesh interpolated from simulation */
 	UPROPERTY(DuplicateTransient, AssetRegistrySearchable)
 	TSoftObjectPtr<UObject> PreviewEmbeddedSkeletalMesh = nullptr;
+
+	/** Static Mesh interpolated from simulation */
+	UPROPERTY(DuplicateTransient, AssetRegistrySearchable)
+	TSoftObjectPtr<UObject> PreviewEmbeddedStaticMesh = nullptr;
 #endif
 
 #if WITH_EDITOR
diff --git a/Engine/Source/Runtime/Experimental/Dataflow/Simulation/Public/Dataflow/Interfaces/DataflowInterfaceGeometryCachable.h b/Engine/Source/Runtime/Experimental/Dataflow/Simulation/Public/Dataflow/Interfaces/DataflowInterfaceGeometryCachable.h
--- a/Engine/Source/Runtime/Experimental/Dataflow/Simulation/Public/Dataflow/Interfaces/DataflowInterfaceGeometryCachable.h
+++ b/Engine/Source/Runtime/Experimental/Dataflow/Simulation/Public/Dataflow/Interfaces/DataflowInterfaceGeometryCachable.h
@@ -11,6 +11,7 @@
 
 class USkeletalMesh;
 class USkinnedAsset;
+class UStaticMesh;
 
 UINTERFACE(meta = (CannotImplementInterfaceInBlueprint), MinimalAPI)
 class UDataflowGeometryCachable : public UInterface
@@ -27,6 +28,9 @@ public:
 	/** Computes skeletal mesh positions for geometry cache */
 	virtual TArray<FVector3f> GetGeometryCachePositions(const USkeletalMesh* SkeletalMesh) const = 0;
 
+	/** Computes static mesh positions for geometry cache */
+	virtual TArray<FVector3f> GetGeometryCachePositions(const UStaticMesh* StaticMesh) const = 0;
+
 	/** Get mapping from final mesh vertex index to raw import vertex index */
 	virtual TOptional<TArray<int32>> GetMeshImportVertexMap(const USkinnedAsset& SkinnedMeshAsset) const = 0;
 
