Commented out USD_FORCE_DISABLED for flesh module

diff --git a/Engine/Plugins/Experimental/ChaosFlesh/Source/ChaosFleshEditor/ChaosFleshEditor.Build.cs b/Engine/Plugins/Experimental/ChaosFlesh/Source/ChaosFleshEditor/ChaosFleshEditor.Build.cs
--- a/Engine/Plugins/Experimental/ChaosFlesh/Source/ChaosFleshEditor/ChaosFleshEditor.Build.cs
+++ b/Engine/Plugins/Experimental/ChaosFlesh/Source/ChaosFleshEditor/ChaosFleshEditor.Build.cs
@@ -139,12 +139,12 @@ namespace UnrealBuildTool.Rules
 			// because those static libraries would need to be linked with the --whole-archive argument, and there is currently no standard way of doing that in UE.
 			if (bEnableUsdSdk && Target.LinkType == TargetLinkType.Monolithic && (Target.StaticAllocator != StaticAllocatorType.Ansi) && !Target.GlobalDefinitions.Contains("UE_USE_MALLOC_FILL_BYTES=0"))
 			{
-				PublicDefinitions.Add("USD_FORCE_DISABLED=1");
+				//PublicDefinitions.Add("USD_FORCE_DISABLED=1");
 				bEnableUsdSdk = false;
 			}
 			else
 			{
-				PublicDefinitions.Add("USD_FORCE_DISABLED=0");
+				//PublicDefinitions.Add("USD_FORCE_DISABLED=0");
 			}
 
 			return bEnableUsdSdk;
diff --git a/Engine/Plugins/Experimental/ChaosFlesh/Source/ChaosFleshEngine/ChaosFleshEngine.Build.cs b/Engine/Plugins/Experimental/ChaosFlesh/Source/ChaosFleshEngine/ChaosFleshEngine.Build.cs
--- a/Engine/Plugins/Experimental/ChaosFlesh/Source/ChaosFleshEngine/ChaosFleshEngine.Build.cs
+++ b/Engine/Plugins/Experimental/ChaosFlesh/Source/ChaosFleshEngine/ChaosFleshEngine.Build.cs
@@ -136,12 +136,12 @@ namespace UnrealBuildTool.Rules
 			// because those static libraries would need to be linked with the --whole-archive argument, and there is currently no standard way of doing that in UE.
 			if (bEnableUsdSdk && Target.LinkType == TargetLinkType.Monolithic && (Target.StaticAllocator != StaticAllocatorType.Ansi) && !Target.GlobalDefinitions.Contains("UE_USE_MALLOC_FILL_BYTES=0"))
 			{
-				PublicDefinitions.Add("USD_FORCE_DISABLED=1");
+				//PublicDefinitions.Add("USD_FORCE_DISABLED=1");
 				bEnableUsdSdk = false;
 			}
 			else
 			{
-				PublicDefinitions.Add("USD_FORCE_DISABLED=0");
+				//PublicDefinitions.Add("USD_FORCE_DISABLED=0");
 			}
 
 			return bEnableUsdSdk;
