ChaosFlesh:
-Added ThresholdScalingPercent in FSetMuscleActivationParameterNode to further control muscle length threshold inferred from input animation.
-Checks for valid origin and insertion indices in FComputeFiberFieldNode

diff --git a/Engine/Plugins/Experimental/ChaosFlesh/Source/ChaosFleshNodes/Private/Dataflow/ChaosFleshComputeFiberFieldNode.cpp b/Engine/Plugins/Experimental/ChaosFlesh/Source/ChaosFleshNodes/Private/Dataflow/ChaosFleshComputeFiberFieldNode.cpp
--- a/Engine/Plugins/Experimental/ChaosFlesh/Source/ChaosFleshNodes/Private/Dataflow/ChaosFleshComputeFiberFieldNode.cpp
+++ b/Engine/Plugins/Experimental/ChaosFlesh/Source/ChaosFleshNodes/Private/Dataflow/ChaosFleshComputeFiberFieldNode.cpp
@@ -142,16 +142,34 @@ void FComputeFiberFieldNode::Evaluate(UE::Dataflow::FContext& Context, const FDa
 		TArray<int32> GeometryIndex = MeshFacade.GetGeometryGroupIndexArray();
 		for (int32 OriginIdx : InOriginIndices)
 		{
+			// Verify origins
+			if (!Vertex->IsValidIndex(OriginIdx))
+			{
+				Context.Error(FString::Printf(
+					TEXT("OriginIdx %d is not a valid vertex index for vertex group size %d."),
+					OriginIdx, Vertex->Num()),
+					this);
+				return;
+			}
 			if (GeometryIndex.IsValidIndex(OriginIdx))
 			{
 				MuscleGeometries.Add(GeometryIndex[OriginIdx]);
 			}
 		}
-		for (int32 OriginIdx : InInsertionIndices)
+		for (int32 InsertionIdx : InInsertionIndices)
 		{
-			if (GeometryIndex.IsValidIndex(OriginIdx))
+			// Verify insertions
+			if (!Vertex->IsValidIndex(InsertionIdx))
 			{
-				MuscleGeometries.Add(GeometryIndex[OriginIdx]);
+				Context.Error(FString::Printf(
+					TEXT("OriginIdx %d is not a valid vertex index for vertex group size %d."),
+					InsertionIdx, Vertex->Num()),
+					this);
+				return;
+			}
+			if (GeometryIndex.IsValidIndex(InsertionIdx))
+			{
+				MuscleGeometries.Add(GeometryIndex[InsertionIdx]);
 			}
 		}
 		TArray<int32> MuscleElementIndices;
diff --git a/Engine/Plugins/Experimental/ChaosFlesh/Source/ChaosFleshNodes/Private/Dataflow/ChaosFleshComputeMuscleActivationNode.cpp b/Engine/Plugins/Experimental/ChaosFlesh/Source/ChaosFleshNodes/Private/Dataflow/ChaosFleshComputeMuscleActivationNode.cpp
--- a/Engine/Plugins/Experimental/ChaosFlesh/Source/ChaosFleshNodes/Private/Dataflow/ChaosFleshComputeMuscleActivationNode.cpp
+++ b/Engine/Plugins/Experimental/ChaosFlesh/Source/ChaosFleshNodes/Private/Dataflow/ChaosFleshComputeMuscleActivationNode.cpp
@@ -229,6 +229,10 @@ FSetMuscleActivationParameterNode::FSetMuscleActivationParameterNode(const UE::D
 					const FManagedArrayCollection& InCollection = GetValue(Context, &Collection);
 					GeometryCollection::Facades::FMuscleActivationFacade MuscleActivation(InCollection);
 					TArray<float> MinLengthRatio;
+					auto ScaledMinLength = [this](float Length) {
+						return 1.f - (1.f - Length) * ThresholdScalingPercent / 100.f;
+					};
+					
 					if (FindLowestMuscleLengthRatio(InCollection, *InAnimationAsset, *InSkeletalMesh, MinLengthRatio)) //success
 					{
 						const int32 NumMuscles = MinLengthRatio.Num();
@@ -238,15 +242,25 @@ FSetMuscleActivationParameterNode::FSetMuscleActivationParameterNode(const UE::D
 							const int32 MuscleIndex = MuscleActivation.FindMuscleIndexByName(MuscleName);
 							if (MinLengthRatio.IsValidIndex(MuscleIndex))
 							{
-								if (MinLengthRatio[MuscleIndex] > 0 && MinLengthRatio[MuscleIndex] < 1)
+								const float ScaledRatio = ScaledMinLength(MinLengthRatio[MuscleIndex]);
+								if (ScaledRatio > 0 && ScaledRatio < 1)
 								{
-									ParameterArray[Idx].MuscleLengthRatioThresholdForMaxActivation = MinLengthRatio[MuscleIndex];
+									ParameterArray[Idx].MuscleLengthRatioThresholdForMaxActivation = ScaledRatio;
 								}
 								else
 								{
-									UE_LOG(LogChaosFlesh, Warning,
-										TEXT("SetMuscleActivationParameter::ImportLowestMuscleLengthRatio: Muscle [%s] index [%d] has minimum origin-insertion length ratio %.2f across the whole animation."),
-										*MuscleName, MuscleIndex, MinLengthRatio[MuscleIndex]);
+									if (MinLengthRatio[MuscleIndex] > 0 && MinLengthRatio[MuscleIndex] < 1)
+									{
+										UE_LOG(LogChaosFlesh, Error,
+											TEXT("SetMuscleActivationParameter::ImportLowestMuscleLengthRatio: Muscle [%s] index [%d] has minimum origin-insertion length ratio %.2f (scaled to %.2f) across the whole animation."),
+											*MuscleName, MuscleIndex, MinLengthRatio[MuscleIndex], ScaledRatio);
+									}
+									else
+									{
+										UE_LOG(LogChaosFlesh, Warning,
+											TEXT("SetMuscleActivationParameter::ImportLowestMuscleLengthRatio: Muscle [%s] index [%d] has minimum origin-insertion length ratio %.2f across the whole animation."),
+											*MuscleName, MuscleIndex, MinLengthRatio[MuscleIndex]);
+									}
 								}
 							}
 							else
@@ -254,7 +268,10 @@ FSetMuscleActivationParameterNode::FSetMuscleActivationParameterNode(const UE::D
 								UE_LOG(LogChaosFlesh, Error,
 									TEXT("SetMuscleActivationParameter::ImportLowestMuscleLengthRatio: Geometry [%s] is not a valid muscle."),
 									*MuscleName);
-								Context.Error(TEXT("ImportLowestMuscleLengthRatio: Geometry [%s] is not a valid muscle."), this);
+								Context.Error(FString::Printf(
+									TEXT("ImportLowestMuscleLengthRatio: Geometry [%s] is not a valid muscle."),
+									*MuscleName),
+									this);
 							}
 						}
 					}
@@ -363,7 +380,9 @@ bool FSetMuscleActivationParameterNode::ShouldInvalidateOnPropertyChanged(const
 	const FName ChangedPropertyName = InPropertyChangedEvent.GetMemberPropertyName();
 	if (ChangedPropertyName == GET_MEMBER_NAME_CHECKED(FSetMuscleActivationParameterNode, ParameterArray) ||
 		ChangedPropertyName == GET_MEMBER_NAME_CHECKED(FSetMuscleActivationParameterNode, ImportAllMuscleNames) ||
-		ChangedPropertyName == GET_MEMBER_NAME_CHECKED(FSetMuscleActivationParameterNode, ResetToGlobalParameters))
+		ChangedPropertyName == GET_MEMBER_NAME_CHECKED(FSetMuscleActivationParameterNode, ResetToGlobalParameters) ||
+		ChangedPropertyName == GET_MEMBER_NAME_CHECKED(FSetMuscleActivationParameterNode, ThresholdScalingPercent) ||
+		ChangedPropertyName == GET_MEMBER_NAME_CHECKED(FSetMuscleActivationParameterNode, ImportLowestMuscleLengthRatio))
 	{
 		return false;
 	}
diff --git a/Engine/Plugins/Experimental/ChaosFlesh/Source/ChaosFleshNodes/Public/Dataflow/ChaosFleshComputeMuscleActivationNode.h b/Engine/Plugins/Experimental/ChaosFlesh/Source/ChaosFleshNodes/Public/Dataflow/ChaosFleshComputeMuscleActivationNode.h
--- a/Engine/Plugins/Experimental/ChaosFlesh/Source/ChaosFleshNodes/Public/Dataflow/ChaosFleshComputeMuscleActivationNode.h
+++ b/Engine/Plugins/Experimental/ChaosFlesh/Source/ChaosFleshNodes/Public/Dataflow/ChaosFleshComputeMuscleActivationNode.h
@@ -129,32 +129,38 @@ public:
 	UPROPERTY(EditAnywhere, Category = "Global Parameters", meta = (ButtonImage = "Icons.Refresh"))
 	FDataflowFunctionProperty ApplyGlobalParameters;
 
-	/* Muscles gain volume during contraction if > 1. Volume-preserving if 1. Default: 1*/
+	/** Muscles gain volume during contraction if > 1. Volume-preserving if 1. Default: 1 */
 	UPROPERTY(EditAnywhere, Category = "Global Parameters", DisplayName = "Global Contraction Volume Scale", meta = (ClampMin = "0"))
 	float ContractionVolumeScale = 1.f;
 
-	/* How much muscle fibers shorten at max activation 1. A smaller value means more contraction in the fiber direction. Default: 0.5 */
+	/** How much muscle fibers shorten at max activation 1. A smaller value means more contraction in the fiber direction. Default: 0.5 */
 	UPROPERTY(EditAnywhere, Category = "Global Parameters", meta = (ClampMin = "0", ClampMax = "1"))
 	float GlobalFiberLengthRatioAtMaxActivation = 0.5f;
 
-	/* Muscle length ratio (defined by origin-insertion distance) below this threshold is considered to reach max activation 1. Default: 0.75 */
+	/** Muscle length ratio (defined by origin-insertion distance) below this threshold is considered to reach max activation 1. Default: 0.75 */
 	UPROPERTY(EditAnywhere, Category = "Global Parameters", meta = (ClampMin = "0", ClampMax = "1"))
 	float GlobalMuscleLengthRatioThresholdForMaxActivation = 0.75f;
 
-	/* Inflates muscle rest volume if > 1 and deflates muscle rest volume if < 1. Default: 1*/
+	/** Inflates muscle rest volume if > 1 and deflates muscle rest volume if < 1. Default: 1 */
 	UPROPERTY(EditAnywhere, Category = "Global Parameters", meta = (ClampMin = "0"))
 	float GlobalInflationVolumeScale = 1.f;
 	
-	/* Use minimum muscle lengths across this animation asset to infer MuscleLengthRatioThresholdForMaxActivation*/
+	/** Use minimum muscle lengths across this animation asset to infer MuscleLengthRatioThresholdForMaxActivation*/
 	UPROPERTY(EditAnywhere, Category = "Animation Data", meta = (DataflowInput))
 	TObjectPtr<UAnimSequence> AnimationAsset = nullptr;
 
-	/* Skeletal mesh used to linear blend skin kinematic origins and insertions. 
+	/** Skeletal mesh used to linear blend skin kinematic origins and insertions. 
 	* This must be the same as the one used to create TransformSource group. 
 	* Check geometry spreadsheet for more info.*/
 	UPROPERTY(EditAnywhere, Category = "Animation Data", meta = (DataflowInput))
 	TObjectPtr<USkeletalMesh> SkeletalMesh = nullptr;
 
+	/** This percentage (usually between 0 and 100%) globally scales minimum muscle lengths computed from the animation asset.
+	* For example, if minimum muscle length = 0.5 and ThresholdScalingPercent = 80%,
+	* the final MuscleLengthRatioThresholdForMaxActivation for this muscle is 1 - (1 - 0.5) * 80% = 0.6 */
+	UPROPERTY(EditAnywhere, Category = "Animation Data", meta = (ClampMin = "0.0", ClampMax = "200.0", UIMin = "0.0", UIMax = "200.0", Units = "Percent"))
+	float ThresholdScalingPercent = 100.f;
+
 	/** Click on this button to import the minimum muscle lengths across the animation asset as MuscleLengthRatioThresholdForMaxActivation*/
 	UPROPERTY(EditAnywhere, Category = "Animation Data", meta = (ButtonImage = "Icons.Import"))
 	FDataflowFunctionProperty ImportLowestMuscleLengthRatio;
