ChaosFlesh: mark simulation dirty if flesh asset properties are changed so simulation scene can be reset.

diff --git a/Engine/Plugins/Experimental/ChaosFlesh/Source/ChaosFleshEngine/Private/ChaosFlesh/FleshAsset.cpp b/Engine/Plugins/Experimental/ChaosFlesh/Source/ChaosFleshEngine/Private/ChaosFlesh/FleshAsset.cpp
--- a/Engine/Plugins/Experimental/ChaosFlesh/Source/ChaosFleshEngine/Private/ChaosFlesh/FleshAsset.cpp
+++ b/Engine/Plugins/Experimental/ChaosFlesh/Source/ChaosFleshEngine/Private/ChaosFlesh/FleshAsset.cpp
@@ -176,7 +176,15 @@ void UDataflowFleshContent::SetActorProperties(TObjectPtr<AActor>& PreviewActor)
 	OverrideStructProperty(PreviewActor, SolverMuscleActivation, TEXT("SolverMuscleActivation"));
 }
 
+#if WITH_EDITOR
+
+void UDataflowFleshContent::PostEditChangeProperty(FPropertyChangedEvent& PropertyChangedEvent)
+{
+	Super::PostEditChangeProperty(PropertyChangedEvent);
+	SetSimulationDirty(true);
+}
 
+#endif //if WITH_EDITOR
 
 
 
diff --git a/Engine/Plugins/Experimental/ChaosFlesh/Source/ChaosFleshEngine/Public/ChaosFlesh/FleshAsset.h b/Engine/Plugins/Experimental/ChaosFlesh/Source/ChaosFleshEngine/Public/ChaosFlesh/FleshAsset.h
--- a/Engine/Plugins/Experimental/ChaosFlesh/Source/ChaosFleshEngine/Public/ChaosFlesh/FleshAsset.h
+++ b/Engine/Plugins/Experimental/ChaosFlesh/Source/ChaosFleshEngine/Public/ChaosFlesh/FleshAsset.h
@@ -207,6 +207,10 @@ public:
 	/** Set all the preview actor exposed properties */
 	virtual void SetActorProperties(TObjectPtr<AActor>& PreviewActor) const override;
 
+#if WITH_EDITOR
+	virtual void PostEditChangeProperty(FPropertyChangedEvent& PropertyChangedEvent) override;
+#endif //if WITH_EDITOR
+
 	UPROPERTY(EditAnywhere, Category = "Preview", Transient, SkipSerialization)
 	FSolverTimingGroup SolverTiming;
 
