ChaosFlesh:
-Force single threaded if there's only one geometry in CreateTetrahedron node
-Use TBitArray instead of TSet.Contains to check orphans in GenerateSurfaceBindings node.
-Add v2 node to SetFleshBonePositionTargetBinding: use SearchRadius instead of bounding box scale and improve binding strategy

diff --git a/Engine/Plugins/Experimental/ChaosFlesh/Source/ChaosFleshNodes/Private/Dataflow/ChaosFleshCreateTetrahedronNode.cpp b/Engine/Plugins/Experimental/ChaosFlesh/Source/ChaosFleshNodes/Private/Dataflow/ChaosFleshCreateTetrahedronNode.cpp
--- a/Engine/Plugins/Experimental/ChaosFlesh/Source/ChaosFleshNodes/Private/Dataflow/ChaosFleshCreateTetrahedronNode.cpp
+++ b/Engine/Plugins/Experimental/ChaosFlesh/Source/ChaosFleshNodes/Private/Dataflow/ChaosFleshCreateTetrahedronNode.cpp
@@ -116,7 +116,7 @@ void FCreateTetrahedronDataflowNode::Evaluate(UE::Dataflow::FContext& Context, c
 					TArray<int32> SortedVertices = VertexToDeleteSet.Array(); SortedVertices.Sort();
 					if (VertexToDeleteSet.Num()) TetCollection.RemoveElements(FGeometryCollection::VerticesGroup, SortedVertices);
 				}
-			});
+			}, ProcessGeometryIndices.Num() < 2? EParallelForFlags::ForceSingleThread : EParallelForFlags::None);
 
 			auto AppendProcessedGeometry = [&ProcessGeometryIndices, &CollectionBuffer, &Source, &Target]()
 			{
diff --git a/Engine/Plugins/Experimental/ChaosFlesh/Source/ChaosFleshNodes/Private/Dataflow/ChaosFleshGenerateSurfaceBindingsNode.cpp b/Engine/Plugins/Experimental/ChaosFlesh/Source/ChaosFleshNodes/Private/Dataflow/ChaosFleshGenerateSurfaceBindingsNode.cpp
--- a/Engine/Plugins/Experimental/ChaosFlesh/Source/ChaosFleshNodes/Private/Dataflow/ChaosFleshGenerateSurfaceBindingsNode.cpp
+++ b/Engine/Plugins/Experimental/ChaosFlesh/Source/ChaosFleshNodes/Private/Dataflow/ChaosFleshGenerateSurfaceBindingsNode.cpp
@@ -593,9 +593,13 @@ FGenerateSurfaceBindings::Evaluate(UE::Dataflow::FContext& Context, const FDataf
 				}
 				const TArray<TArray<uint32>>& NeighborNodes = RenderMeshNeighborNodes[LOD];
 				OrphansCopy = Orphans;
-				TSet<int32> OrphanSet(Orphans);
+				TBitArray<> IsOrphan(false, RenderMeshVertices[LOD].Num());
 				while (bDoOrphanReparenting && Orphans.Num())
 				{
+					for (int32 Orphan : Orphans)
+					{
+						IsOrphan[Orphan] = true;
+					}
 					// Find the orphan with the fewest number of orphan neighbors, and the 
 					// most non-orphans in their 1 ring.
 					int32 Orphan = INDEX_NONE;
@@ -613,7 +617,7 @@ FGenerateSurfaceBindings::Evaluate(UE::Dataflow::FContext& Context, const FDataf
 						int32 NonOrphanCount = 0;
 						for (int32 j = 0; j < Neighbors.Num(); j++)
 						{
-							if (OrphanSet.Contains(Neighbors[j]))
+							if (IsOrphan[Neighbors[j]])
 							{
 								OrphanCount++;
 							}
@@ -644,7 +648,7 @@ FGenerateSurfaceBindings::Evaluate(UE::Dataflow::FContext& Context, const FDataf
 					for (int32 i = 0; i < Neighbors.Num(); i++)
 					{
 						const uint32 Neighbor = Neighbors[i];
-						if (OrphanSet.Contains(Neighbor))
+						if (IsOrphan[Neighbor])
 						{
 							continue;
 						}
@@ -698,7 +702,7 @@ FGenerateSurfaceBindings::Evaluate(UE::Dataflow::FContext& Context, const FDataf
 					} // end for all neighbors
 
 					// Whether or not we successfully reparented, remove the orphan from the list.
-					OrphanSet.Remove(Orphan);
+					IsOrphan[Orphan] = false;
 					Orphans.Remove(Orphan);
 					if (FoundBinding)
 					{
diff --git a/Engine/Plugins/Experimental/ChaosFlesh/Source/ChaosFleshNodes/Private/Dataflow/ChaosFleshNodesPlugin.cpp b/Engine/Plugins/Experimental/ChaosFlesh/Source/ChaosFleshNodes/Private/Dataflow/ChaosFleshNodesPlugin.cpp
--- a/Engine/Plugins/Experimental/ChaosFlesh/Source/ChaosFleshNodes/Private/Dataflow/ChaosFleshNodesPlugin.cpp
+++ b/Engine/Plugins/Experimental/ChaosFlesh/Source/ChaosFleshNodes/Private/Dataflow/ChaosFleshNodesPlugin.cpp
@@ -72,6 +72,7 @@ void IChaosFleshNodesPlugin::StartupModule()
 	DATAFLOW_NODE_REGISTER_CREATION_FACTORY(FReadSkeletalMeshCurvesDataflowNode);
 	DATAFLOW_NODE_REGISTER_CREATION_FACTORY(FSetCollidableVerticesDataflowNode);
 	DATAFLOW_NODE_REGISTER_CREATION_FACTORY(FSetFleshBonePositionTargetBindingDataflowNode);
+	DATAFLOW_NODE_REGISTER_CREATION_FACTORY(FSetFleshBonePositionTargetBindingDataflowNode_v2);
 	DATAFLOW_NODE_REGISTER_CREATION_FACTORY(FSetFleshDefaultPropertiesNode);
 	DATAFLOW_NODE_REGISTER_CREATION_FACTORY(FSetMuscleActivationParameterNode);
 	DATAFLOW_NODE_REGISTER_CREATION_FACTORY(FSetVerticesKinematicDataflowNode);
diff --git a/Engine/Plugins/Experimental/ChaosFlesh/Source/ChaosFleshNodes/Private/Dataflow/ChaosFleshSetFleshBonePositionTargetBindingNode.cpp b/Engine/Plugins/Experimental/ChaosFlesh/Source/ChaosFleshNodes/Private/Dataflow/ChaosFleshSetFleshBonePositionTargetBindingNode.cpp
--- a/Engine/Plugins/Experimental/ChaosFlesh/Source/ChaosFleshNodes/Private/Dataflow/ChaosFleshSetFleshBonePositionTargetBindingNode.cpp
+++ b/Engine/Plugins/Experimental/ChaosFlesh/Source/ChaosFleshNodes/Private/Dataflow/ChaosFleshSetFleshBonePositionTargetBindingNode.cpp
@@ -13,18 +13,20 @@
 #include "Rendering/SkeletalMeshRenderData.h"
 #include "SkeletalMeshAttributes.h"
 #include "ChaosFlesh/ChaosFlesh.h"
+#include "Chaos/HierarchicalSpatialHash.h"
+#include "Chaos/TriangleCollisionPoint.h"
 
 #include "BoneWeights.h"
 
 #include UE_INLINE_GENERATED_CPP_BY_NAME(ChaosFleshSetFleshBonePositionTargetBindingNode)
 
-//DEFINE_LOG_CATEGORY_STATIC(ChaosFleshSetFleshBonePositionTargetBindingNodeLog, Log, All);
+DEFINE_LOG_CATEGORY_STATIC(ChaosFleshSetFleshBonePositionTargetBindingNodeLog, Log, All);
 
 void FSetFleshBonePositionTargetBindingDataflowNode::Evaluate(UE::Dataflow::FContext& Context, const FDataflowOutput* Out) const
 {
-	if (Out->IsA<DataType>(&Collection))
+	if (Out->IsA(&Collection))
 	{
-		DataType InCollection = GetValue<DataType>(Context, &Collection);
+		FManagedArrayCollection InCollection = GetValue(Context, &Collection);
 #if WITH_EDITORONLY_DATA
 		if (TObjectPtr<USkeletalMesh> BoneSkeletalMesh = GetValue<TObjectPtr<USkeletalMesh>>(Context, &SkeletalMeshIn))
 		{
@@ -79,7 +81,7 @@ void FSetFleshBonePositionTargetBindingDataflowNode::Evaluate(UE::Dataflow::FCon
 
 					GeometryCollection::Facades::FVertexBoneWeightsFacade VertexBoneWeightsFacade(InCollection);
 
-					
+
 
 					//Mapping from SKM bone index to Collection bone index
 					const TMap<FString, int32> BoneNameIndexMap = TransformFacade.BoneNameIndexMap();
@@ -135,7 +137,6 @@ void FSetFleshBonePositionTargetBindingDataflowNode::Evaluate(UE::Dataflow::FCon
 					}
 					else
 					{
-						TMap<int32, int32> BoneVertexToCollectionMap;
 						GeometryCollection::Facades::FPositionTargetFacade PositionTargets(InCollection);
 						PositionTargets.DefineSchema();
 
@@ -207,4 +208,164 @@ void FSetFleshBonePositionTargetBindingDataflowNode::Evaluate(UE::Dataflow::FCon
 #endif
 		SetValue(Context, MoveTemp(InCollection), &Collection);
 	}
+}
+
+void FSetFleshBonePositionTargetBindingDataflowNode_v2::Evaluate(UE::Dataflow::FContext& Context, const FDataflowOutput* Out) const
+{
+	if (Out->IsA(&Collection))
+	{
+		FManagedArrayCollection InCollection = GetValue(Context, &Collection);
+#if WITH_EDITORONLY_DATA
+		if (TObjectPtr<USkeletalMesh> BoneSkeletalMesh = GetValue<TObjectPtr<USkeletalMesh>>(Context, &SkeletalMeshIn))
+		{
+			const TManagedArray<FVector3f>* Vertices = InCollection.FindAttribute<FVector3f>("Vertex", FGeometryCollection::VerticesGroup);
+			const TManagedArray<FTransform3f>* Transform = InCollection.FindAttribute<FTransform3f>("Transform", FTransformCollection::TransformGroup);
+			const TManagedArray<FString>* TransformBoneName = InCollection.FindAttribute<FString>("BoneName", FTransformCollection::TransformGroup);
+			GeometryCollection::Facades::FCollectionTransformFacade TransformFacade(InCollection);
+			const FReferenceSkeleton& ReferenceSkeleton = BoneSkeletalMesh->GetRefSkeleton();
+			static int32 LODIndex = 0;
+			if (Vertices && Transform && TransformBoneName)
+			{
+				FMeshDescription MeshDescription;
+				if (BoneSkeletalMesh->CloneMeshDescription(LODIndex, MeshDescription))
+				{
+					FSkeletalMeshAttributes MeshAttribs(MeshDescription);
+					const FSkinWeightsVertexAttributesRef VertexSkinWeights = MeshAttribs.GetVertexSkinWeights();
+					const int32 NumSkeletonVertices = MeshDescription.Vertices().Num();
+					const int32 NumSkeletonTriangles = MeshDescription.Triangles().Num();
+
+					TArray<Chaos::TVec3<Chaos::FReal>> SkeletonVerticesTVec3;
+					SkeletonVerticesTVec3.SetNum(NumSkeletonVertices);
+					// Get Vertex Position
+					for (const FVertexID& VertexID : MeshDescription.Vertices().GetElementIDs())
+					{
+						const int32 VertexIndex = VertexID.GetValue();			
+						SkeletonVerticesTVec3[VertexIndex] = Chaos::TVec3<Chaos::FReal>(MeshDescription.GetVertexPosition(VertexID));
+					}
+					TConstArrayView<Chaos::TVec3<Chaos::FReal>> ConstSkeletonVerticesTVec3(SkeletonVerticesTVec3);
+					// Get Triangles
+					Chaos::FTriangleMesh TriangleMesh;
+					TArray<Chaos::TVec3<int32>> TriangleElements;
+					TriangleElements.SetNum(NumSkeletonTriangles);
+					for (const FTriangleID& TriangleID : MeshDescription.Triangles().GetElementIDs())
+					{
+						const int32 TriangleIndex = TriangleID.GetValue();
+						const TArrayView<const FVertexID>& Triangle = MeshDescription.GetTriangleVertices(TriangleID);
+						TriangleElements[TriangleIndex] = Chaos::TVec3<int32>(Triangle[0], Triangle[1], Triangle[2]);
+					}
+					TriangleMesh.Init(TriangleElements);
+					Chaos::FTriangleMesh::TSpatialHashType<Chaos::FReal> SpatialHash;
+					Chaos::FReal SphereRadius = (Chaos::FReal)SearchRadius;
+					TriangleMesh.BuildSpatialHash(ConstSkeletonVerticesTVec3, SpatialHash, SphereRadius);
+
+					GeometryCollection::Facades::FVertexBoneWeightsFacade VertexBoneWeightsFacade(InCollection);
+					GeometryCollection::Facades::FPositionTargetFacade PositionTargets(InCollection);
+
+					//Mapping from SKM bone index to Collection bone index
+					const TMap<FString, int32> BoneNameIndexMap = TransformFacade.BoneNameIndexMap();
+					TMap<int32, int32> SKMBoneIndexToCollectionBoneIndex;
+					for (int32 BoneIndex = 0; BoneIndex < ReferenceSkeleton.GetNum(); ++BoneIndex)
+					{
+						FString SKMBoneName = ReferenceSkeleton.GetRefBoneInfo()[BoneIndex].Name.ToString();
+						if (BoneNameIndexMap.Contains(SKMBoneName))
+						{
+							SKMBoneIndexToCollectionBoneIndex.Add(BoneIndex, BoneNameIndexMap[SKMBoneName]);
+						}
+					}
+
+					//Transfer skin weights
+					TSet<int32> MissingSourceBones;
+					for (int32 VertIdx = 0; VertIdx < Vertices->Num(); ++VertIdx)
+					{
+						//only work on particles that are not kinematic:
+						if (!VertexBoneWeightsFacade.IsKinematicVertex(VertIdx))
+						{
+							TArray<Chaos::TTriangleCollisionPoint<Chaos::FReal>> Result;
+							if (TriangleMesh.PointClosestTriangleQuery(SpatialHash, ConstSkeletonVerticesTVec3,
+								VertIdx, Chaos::TVec3<Chaos::FReal>((*Vertices)[VertIdx]), SphereRadius / 2.f, SphereRadius / 2.f,
+								[](const int32 PointIndex, const int32 TriangleIndex)->bool { return true; }, Result))
+							{
+								for (const Chaos::TTriangleCollisionPoint<Chaos::FReal>& CollisionPoint : Result)
+								{
+									TMap<int32, float> BoneWeightBucket;
+									FVector3f InterpSkeletonVertexPosition = FVector3f(0);
+									for (int32 LocalTriIdx = 0; LocalTriIdx < 3; ++LocalTriIdx)
+									{
+										const int32 TriVertexIndex = TriangleElements[CollisionPoint.Indices[1]][LocalTriIdx];
+										const float TriInterpWeight = CollisionPoint.Bary[LocalTriIdx + 1];
+										InterpSkeletonVertexPosition += TriInterpWeight * MeshDescription.GetVertexPosition(TriVertexIndex);
+										FVertexBoneWeights BoneWeights = VertexSkinWeights.Get(FVertexID(TriVertexIndex));
+										const int32 InfluenceCount = BoneWeights.Num();
+										for (int32 InfluenceIndex = 0; InfluenceIndex < InfluenceCount; ++InfluenceIndex)
+										{
+											const int32 BoneIndex = SKMBoneIndexToCollectionBoneIndex[BoneWeights[InfluenceIndex].GetBoneIndex()];
+											if (SKMBoneIndexToCollectionBoneIndex.Contains(BoneWeights[InfluenceIndex].GetBoneIndex()))
+											{
+												const float InterpBoneWeight = TriInterpWeight * BoneWeights[InfluenceIndex].GetWeight();
+												if (BoneWeightBucket.Contains(BoneIndex))
+												{
+													BoneWeightBucket[BoneIndex] += InterpBoneWeight;
+												}
+												else
+												{
+													BoneWeightBucket.Add(BoneIndex, InterpBoneWeight);
+												}
+											}
+											else if (!MissingSourceBones.Contains(BoneIndex))
+											{
+												MissingSourceBones.Add(BoneIndex);
+												UE_LOG(ChaosFleshSetFleshBonePositionTargetBindingNodeLog, Error, TEXT("Collection does not contain bone[%s]."), *BoneSkeletalMesh->GetRefSkeleton().GetBoneName(BoneIndex).ToString());
+											}
+										}
+									}
+									TArray<int32> VertexBoneIndex;
+									TArray<float> VertexBoneWeight;
+									for (const TPair<int32, float>& Pair : BoneWeightBucket)
+									{
+										VertexBoneIndex.Add(Pair.Key);
+										VertexBoneWeight.Add(Pair.Value);
+									}
+									if (SkeletalBindingMode == ESkeletalBindingMode::Dataflow_SkeletalBinding_Kinematic)
+									{
+										VertexBoneWeightsFacade.ModifyBoneWeight(VertIdx, VertexBoneIndex, VertexBoneWeight);
+										VertexBoneWeightsFacade.SetVertexKinematic(VertIdx);
+									}
+									else // Dataflow_SkeletalBinding_PositionTarget
+									{
+										const int32 ParticleIndex = InCollection.AddElements(1, FGeometryCollection::VerticesGroup);
+										TManagedArray<FVector3f>& CurrentVertices = InCollection.ModifyAttribute<FVector3f>("Vertex", FGeometryCollection::VerticesGroup);
+										CurrentVertices[ParticleIndex] = InterpSkeletonVertexPosition;
+
+										VertexBoneWeightsFacade.ModifyBoneWeight(ParticleIndex, VertexBoneIndex, VertexBoneWeight);
+										VertexBoneWeightsFacade.SetVertexKinematic(ParticleIndex);
+
+										GeometryCollection::Facades::FPositionTargetsData DataPackage;
+										DataPackage.TargetIndex.Init(ParticleIndex, 1);
+										DataPackage.TargetWeights.Init(1.f, 1);
+										DataPackage.SourceWeights.Init(1.f, 1);
+										DataPackage.SourceIndex.Init(VertIdx, 1);
+										if (const TManagedArray<float>* Mass = InCollection.FindAttribute<float>("Mass", FGeometryCollection::VerticesGroup))
+										{
+											//Target is kinematic, only compute stiffness from source
+											DataPackage.Stiffness = PositionTargetStiffness * (*Mass)[VertIdx];
+										}
+										else
+										{
+											DataPackage.Stiffness = PositionTargetStiffness;
+										}
+										PositionTargets.AddPositionTarget(DataPackage);
+									}
+									break;
+								}
+							}
+						}
+					}
+				}
+			}
+		}
+#else
+		ensureMsgf(false, TEXT("SetFleshBonePositionTargetBinding is an editor only node."));
+#endif
+		SetValue(Context, MoveTemp(InCollection), &Collection);
+	}
 }
\ No newline at end of file
diff --git a/Engine/Plugins/Experimental/ChaosFlesh/Source/ChaosFleshNodes/Public/Dataflow/ChaosFleshGenerateSurfaceBindingsNode.h b/Engine/Plugins/Experimental/ChaosFlesh/Source/ChaosFleshNodes/Public/Dataflow/ChaosFleshGenerateSurfaceBindingsNode.h
--- a/Engine/Plugins/Experimental/ChaosFlesh/Source/ChaosFleshNodes/Public/Dataflow/ChaosFleshGenerateSurfaceBindingsNode.h
+++ b/Engine/Plugins/Experimental/ChaosFlesh/Source/ChaosFleshNodes/Public/Dataflow/ChaosFleshGenerateSurfaceBindingsNode.h
@@ -71,7 +71,7 @@ private:
 	TArray<FString> GeometryGroupGuidsIn;
 
 	/** Use the import geometry of the skeletal mesh. */
-	UPROPERTY(EditAnywhere, Category = "Dataflow", meta = (DisplayName = "UseSkeletalMeshImportModel", EditCondition = "MeshInputType == EChaosFleshGenerateSurfaceBindingsMeshInputType::SkeletalMesh", EditConditionHides))
+	UPROPERTY(EditAnywhere, Category = "Dataflow", meta = (DisplayName = "UseSkeletalMeshImportModel", EditCondition = "SkeletalMeshIn != nullptr"))
 	bool bUseSkeletalMeshImportModel = false;
 
 	/** Enable binding to the exterior hull of the tetrahedron mesh. */
diff --git a/Engine/Plugins/Experimental/ChaosFlesh/Source/ChaosFleshNodes/Public/Dataflow/ChaosFleshSetFleshBonePositionTargetBindingNode.h b/Engine/Plugins/Experimental/ChaosFlesh/Source/ChaosFleshNodes/Public/Dataflow/ChaosFleshSetFleshBonePositionTargetBindingNode.h
--- a/Engine/Plugins/Experimental/ChaosFlesh/Source/ChaosFleshNodes/Public/Dataflow/ChaosFleshSetFleshBonePositionTargetBindingNode.h
+++ b/Engine/Plugins/Experimental/ChaosFlesh/Source/ChaosFleshNodes/Public/Dataflow/ChaosFleshSetFleshBonePositionTargetBindingNode.h
@@ -19,31 +19,31 @@ enum class ESkeletalBindingMode : uint8
 	Chaos_Max UMETA(Hidden)
 };
 
-USTRUCT(meta = (DataflowFlesh))
+/**
+* Deprecated in 5.6 to support fixed distance search radius and improve binding method.
+*/
+USTRUCT(meta = (Deprecated = "5.6"))
 struct FSetFleshBonePositionTargetBindingDataflowNode : public FDataflowNode
 {
 	GENERATED_USTRUCT_BODY()
-		DATAFLOW_NODE_DEFINE_INTERNAL(FSetFleshBonePositionTargetBindingDataflowNode, "SetFleshBonePositionTargetBinding", "Flesh", "")
-		DATAFLOW_NODE_RENDER_TYPE("SurfaceRender",FGeometryCollection::StaticType(),  "Collection")
+	DATAFLOW_NODE_DEFINE_INTERNAL(FSetFleshBonePositionTargetBindingDataflowNode, "SetFleshBonePositionTargetBinding", "Flesh", "")
+	DATAFLOW_NODE_RENDER_TYPE("SurfaceRender", FGeometryCollection::StaticType(), "Collection")
 
 public:
-	typedef FManagedArrayCollection DataType;
-
-
 	UPROPERTY(meta = (DataflowInput, DataflowOutput, DisplayName = "Collection", DataflowPassthrough = "Collection"))
-		FManagedArrayCollection Collection;
+	FManagedArrayCollection Collection;
 
 	UPROPERTY(EditAnywhere, Category = "Dataflow")
-		ESkeletalBindingMode SkeletalBindingMode = ESkeletalBindingMode::Dataflow_SkeletalBinding_PositionTarget;
+	ESkeletalBindingMode SkeletalBindingMode = ESkeletalBindingMode::Dataflow_SkeletalBinding_PositionTarget;
 
 	UPROPERTY(EditAnywhere, Category = "Dataflow", meta = (ClampMin = "0.0", EditCondition = "SkeletalBindingMode == ESkeletalBindingMode::Dataflow_SkeletalBinding_PositionTarget", EditConditionHides))
-		float PositionTargetStiffness = 10000.f;
+	float PositionTargetStiffness = 10000.f;
 
 	UPROPERTY(EditAnywhere, Category = "Dataflow", meta = (DataflowInput, DisplayName = "SkeletalMesh"))
-		TObjectPtr<USkeletalMesh> SkeletalMeshIn = nullptr;
+	TObjectPtr<USkeletalMesh> SkeletalMeshIn = nullptr;
 
 	UPROPERTY(EditAnywhere, Category = "Dataflow", meta = (ClampMin = "0.0", ClampMax = "2.0", UIMin = "0.0", UIMax = "2.0"))
-		float VertexRadiusRatio = .001f;
+	float VertexRadiusRatio = .001f;
 
 
 	FSetFleshBonePositionTargetBindingDataflowNode(const UE::Dataflow::FNodeParameters& InParam, FGuid InGuid = FGuid::NewGuid())
@@ -57,6 +57,47 @@ public:
 	virtual void Evaluate(UE::Dataflow::FContext& Context, const FDataflowOutput* Out) const override;
 };
 
+/**
+* Binds vertices from Collection to bone skeletal mesh surface via kinematic or weak constraints.
+*/
+USTRUCT(meta = (DataflowFlesh))
+struct FSetFleshBonePositionTargetBindingDataflowNode_v2 : public FDataflowNode
+{
+	GENERATED_USTRUCT_BODY()
+	DATAFLOW_NODE_DEFINE_INTERNAL(FSetFleshBonePositionTargetBindingDataflowNode_v2, "SetFleshBonePositionTargetBinding", "Flesh", "")
+	DATAFLOW_NODE_RENDER_TYPE("SurfaceRender",FGeometryCollection::StaticType(),  "Collection")
+
+public:
+	FSetFleshBonePositionTargetBindingDataflowNode_v2(const UE::Dataflow::FNodeParameters& InParam, FGuid InGuid = FGuid::NewGuid())
+		: FDataflowNode(InParam, InGuid)
+	{
+		RegisterInputConnection(&Collection);
+		RegisterInputConnection(&SkeletalMeshIn);
+		RegisterOutputConnection(&Collection, &Collection);
+	}
+
+private:
+	UPROPERTY(meta = (DataflowInput, DataflowOutput, DisplayName = "Collection", DataflowPassthrough = "Collection"))
+	FManagedArrayCollection Collection;
+
+	UPROPERTY(EditAnywhere, Category = "Dataflow")
+	ESkeletalBindingMode SkeletalBindingMode = ESkeletalBindingMode::Dataflow_SkeletalBinding_PositionTarget;
+
+	UPROPERTY(EditAnywhere, Category = "Dataflow", meta = (ClampMin = "0.0", EditCondition = "SkeletalBindingMode == ESkeletalBindingMode::Dataflow_SkeletalBinding_PositionTarget", EditConditionHides))
+	float PositionTargetStiffness = 10000.f;
+
+	UPROPERTY(EditAnywhere, Category = "Dataflow", meta = (DataflowInput, DisplayName = "SkeletalMesh"))
+	TObjectPtr<USkeletalMesh> SkeletalMeshIn = nullptr;
+
+	/**
+	* Collection vertices are bound to their closest skeletal mesh vertices within the search radius
+	*/
+	UPROPERTY(EditAnywhere, Category = "Dataflow", meta = (ClampMin = "0.0"))
+	float SearchRadius = 0.f;
+
+	virtual void Evaluate(UE::Dataflow::FContext& Context, const FDataflowOutput* Out) const override;
+};
+
 namespace UE::Dataflow
 {
 	void RegisterChaosFleshPositionTargetInitializationNodes();
