ChaosFlesh:
-Supported vertex selection in kinematic bindings
-Modified inflation parameter usage
-Fixed geom cache generation on time ranges starting from nonzero

diff --git a/Engine/Plugins/Experimental/ChaosFlesh/Source/ChaosFleshNodes/Private/Dataflow/ChaosFleshSetFleshBonePositionTargetBindingNode.cpp b/Engine/Plugins/Experimental/ChaosFlesh/Source/ChaosFleshNodes/Private/Dataflow/ChaosFleshSetFleshBonePositionTargetBindingNode.cpp
--- a/Engine/Plugins/Experimental/ChaosFlesh/Source/ChaosFleshNodes/Private/Dataflow/ChaosFleshSetFleshBonePositionTargetBindingNode.cpp
+++ b/Engine/Plugins/Experimental/ChaosFlesh/Source/ChaosFleshNodes/Private/Dataflow/ChaosFleshSetFleshBonePositionTargetBindingNode.cpp
@@ -273,9 +273,33 @@ void FSetFleshBonePositionTargetBindingDataflowNode_v2::Evaluate(UE::Dataflow::F
 						}
 					}
 
+					// Only keep boundary vertices within VertexSelection
+					TArray<int32> SelectedVertices;
+					if (IsConnected(&VertexSelection))
+					{
+						FDataflowVertexSelection InDataflowVertexSelection = GetValue<FDataflowVertexSelection>(Context, &VertexSelection);
+						if (InDataflowVertexSelection.Num() != Vertices->Num())
+						{
+							Context.Error(FString::Printf(
+								TEXT("VertexSelection size [%d] is not equal to the collection's vertex count [%d]"),
+								InDataflowVertexSelection.Num(), Vertices->Num()),
+								this, Out);
+							return;
+						}
+						SelectedVertices = InDataflowVertexSelection.AsArray();
+					}
+					else
+					{
+						SelectedVertices.SetNum(Vertices->Num());
+						for (int32 VertIdx = 0; VertIdx < Vertices->Num(); ++VertIdx)
+						{
+							SelectedVertices[VertIdx] = VertIdx;
+						}
+					}
+
 					//Transfer skin weights
 					TSet<int32> MissingSourceBones;
-					for (int32 VertIdx = 0; VertIdx < Vertices->Num(); ++VertIdx)
+					for (int32 VertIdx : SelectedVertices)
 					{
 						//only work on particles that are not kinematic:
 						if (!VertexBoneWeightsFacade.IsKinematicVertex(VertIdx))
@@ -298,23 +322,24 @@ void FSetFleshBonePositionTargetBindingDataflowNode_v2::Evaluate(UE::Dataflow::F
 										const int32 InfluenceCount = BoneWeights.Num();
 										for (int32 InfluenceIndex = 0; InfluenceIndex < InfluenceCount; ++InfluenceIndex)
 										{
-											const int32 BoneIndex = SKMBoneIndexToCollectionBoneIndex[BoneWeights[InfluenceIndex].GetBoneIndex()];
-											if (SKMBoneIndexToCollectionBoneIndex.Contains(BoneWeights[InfluenceIndex].GetBoneIndex()))
+											const int32 SKMBoneIndex = BoneWeights[InfluenceIndex].GetBoneIndex();
+											if (SKMBoneIndexToCollectionBoneIndex.Contains(SKMBoneIndex))
 											{
+												const int32 CollectionBoneIndex = SKMBoneIndexToCollectionBoneIndex[SKMBoneIndex];
 												const float InterpBoneWeight = TriInterpWeight * BoneWeights[InfluenceIndex].GetWeight();
-												if (BoneWeightBucket.Contains(BoneIndex))
+												if (BoneWeightBucket.Contains(CollectionBoneIndex))
 												{
-													BoneWeightBucket[BoneIndex] += InterpBoneWeight;
+													BoneWeightBucket[CollectionBoneIndex] += InterpBoneWeight;
 												}
 												else
 												{
-													BoneWeightBucket.Add(BoneIndex, InterpBoneWeight);
+													BoneWeightBucket.Add(CollectionBoneIndex, InterpBoneWeight);
 												}
 											}
-											else if (!MissingSourceBones.Contains(BoneIndex))
+											else if (!MissingSourceBones.Contains(SKMBoneIndex))
 											{
-												MissingSourceBones.Add(BoneIndex);
-												UE_LOG(ChaosFleshSetFleshBonePositionTargetBindingNodeLog, Error, TEXT("Collection does not contain bone[%s]."), *BoneSkeletalMesh->GetRefSkeleton().GetBoneName(BoneIndex).ToString());
+												MissingSourceBones.Add(SKMBoneIndex);
+												UE_LOG(ChaosFleshSetFleshBonePositionTargetBindingNodeLog, Error, TEXT("Collection does not contain bone[%s]."), *BoneSkeletalMesh->GetRefSkeleton().GetBoneName(SKMBoneIndex).ToString());
 											}
 										}
 									}
diff --git a/Engine/Plugins/Experimental/ChaosFlesh/Source/ChaosFleshNodes/Public/Dataflow/ChaosFleshSetFleshBonePositionTargetBindingNode.h b/Engine/Plugins/Experimental/ChaosFlesh/Source/ChaosFleshNodes/Public/Dataflow/ChaosFleshSetFleshBonePositionTargetBindingNode.h
--- a/Engine/Plugins/Experimental/ChaosFlesh/Source/ChaosFleshNodes/Public/Dataflow/ChaosFleshSetFleshBonePositionTargetBindingNode.h
+++ b/Engine/Plugins/Experimental/ChaosFlesh/Source/ChaosFleshNodes/Public/Dataflow/ChaosFleshSetFleshBonePositionTargetBindingNode.h
@@ -73,6 +73,7 @@ public:
 	{
 		RegisterInputConnection(&Collection);
 		RegisterInputConnection(&SkeletalMeshIn);
+		RegisterInputConnection(&VertexSelection);
 		RegisterOutputConnection(&Collection, &Collection);
 	}
 
@@ -89,6 +90,10 @@ private:
 	UPROPERTY(EditAnywhere, Category = "Dataflow", meta = (DataflowInput, DisplayName = "SkeletalMesh"))
 	TObjectPtr<USkeletalMesh> SkeletalMeshIn = nullptr;
 
+	/** (optional) only create kinematic/weak constraints on vertices in the VertexSelection */
+	UPROPERTY(meta = (DataflowInput, DisplayName = "(Optional) VertexSelection"))
+	FDataflowVertexSelection VertexSelection;
+
 	/**
 	* Collection vertices are bound to their closest skeletal mesh vertices within the search radius
 	*/
diff --git a/Engine/Plugins/Experimental/ChaosFlesh/Source/ChaosFleshNodes/Public/Dataflow/ChaosFleshSetFleshDefaultPropertiesNode.h b/Engine/Plugins/Experimental/ChaosFlesh/Source/ChaosFleshNodes/Public/Dataflow/ChaosFleshSetFleshDefaultPropertiesNode.h
--- a/Engine/Plugins/Experimental/ChaosFlesh/Source/ChaosFleshNodes/Public/Dataflow/ChaosFleshSetFleshDefaultPropertiesNode.h
+++ b/Engine/Plugins/Experimental/ChaosFlesh/Source/ChaosFleshNodes/Public/Dataflow/ChaosFleshSetFleshDefaultPropertiesNode.h
@@ -29,13 +29,13 @@ public:
 	float VertexDamping = 0.f;
 
 	/*Sets incompressibility on vertex basis. 0.6 is default behavior. 
-	1 means total incompressibility. 0.00001 means almost no incompressibility*/
+	1 means absolutely incompressible. 0 means no incompressibility constraint on the material */
 	UPROPERTY(EditAnywhere, Category = "Dataflow", meta = (ClampMin = "0.00001", ClampMax = "1.0", UIMin = "0.00001", UIMax = "1.0"))
 	float VertexIncompressibility = 0.6f;
 
 	/*Sets inflation on vertex basis. 0.5 means no inflation/deflation.
-	1 means total inflation. 0 means the material is deflated.*/
-	UPROPERTY(EditAnywhere, Category = "Dataflow", meta = (ClampMin = "0.0", ClampMax = "1.0", UIMin = "0.0", UIMax = "1.0"))
+	1 means inflation to 2X volume on each dimension. 0 means the material is deflated to 0 volume.*/
+	UPROPERTY(EditAnywhere, Category = "Dataflow", meta = (ClampMin = "0.00001", UIMin = "0.00001"))
 	float VertexInflation = 0.5f;
 
 	UPROPERTY(meta = (DataflowInput, DataflowOutput, DisplayName = "Collection", DataflowPassthrough  = "Collection"))
diff --git a/Engine/Plugins/Experimental/Dataflow/Source/DataflowEditor/Private/Dataflow/DataflowSimulationGenerator.cpp b/Engine/Plugins/Experimental/Dataflow/Source/DataflowEditor/Private/Dataflow/DataflowSimulationGenerator.cpp
--- a/Engine/Plugins/Experimental/Dataflow/Source/DataflowEditor/Private/Dataflow/DataflowSimulationGenerator.cpp
+++ b/Engine/Plugins/Experimental/Dataflow/Source/DataflowEditor/Private/Dataflow/DataflowSimulationGenerator.cpp
@@ -97,7 +97,7 @@ bool FDataflowTaskManager::AllocateSimulationResource(const FVector2f& TimeRange
 	SimulationResource->RecentDateTimeTicks = &RecentDateTimeTicks;
 	SimulationResource->bCancelled = &bCancelled;
 	
-	NumFrames = FMath::Floor((TimeRange[1] - TimeRange[0] + UE_SMALL_NUMBER) * FrameRate);
+	NumFrames = FMath::Floor((TimeRange[1] - TimeRange[0] + UE_KINDA_SMALL_NUMBER) * FrameRate);
 	if(SimulationTask.IsValid())
 	{
 		SimulationTask->GetTask().SimulationWorld = SimulationWorld;
diff --git a/Engine/Plugins/Experimental/Dataflow/Source/DataflowEditor/Private/Dataflow/DataflowSimulationScene.cpp b/Engine/Plugins/Experimental/Dataflow/Source/DataflowEditor/Private/Dataflow/DataflowSimulationScene.cpp
--- a/Engine/Plugins/Experimental/Dataflow/Source/DataflowEditor/Private/Dataflow/DataflowSimulationScene.cpp
+++ b/Engine/Plugins/Experimental/Dataflow/Source/DataflowEditor/Private/Dataflow/DataflowSimulationScene.cpp
@@ -484,9 +484,8 @@ void UDataflowSimulationSceneDescription::GenerateGeometryCache()
 {
 	SimulationScene->ResetSimulationScene();
 	SimulationScene->CreateSimulationScene();
-	const FVector2f& TimeRange = SimulationScene->GetTimeRange();
 	const int32 NumFrames = SimulationScene->GetNumFrames();
-	float Time = TimeRange[0];
+	float Time = 0;
 	TObjectPtr<AActor> GetRootActor = SimulationScene->GetRootActor();
 	TObjectPtr<AActor> PreviewActor = SimulationScene->GetPreviewActor();
 	const bool UseSkeletalMesh = EmbeddedSkeletalMesh != nullptr;
@@ -513,7 +512,7 @@ void UDataflowSimulationSceneDescription::GenerateGeometryCache()
 		}
 		for (int32 Frame = 0; Frame < NumFrames; ++Frame)
 		{
-			Cast<AChaosCacheManager>(GetRootActor)->SetStartTime(Time);
+			Cast<AChaosCacheManager>(GetRootActor)->SetStartTime(Time); // cache time range is [0, (NumFrames-1)*dt]
 			if (UseSkeletalMesh)
 			{
 				RenderPositions[Frame] = GeometryCachable->GetGeometryCachePositions(EmbeddedSkeletalMesh);
diff --git a/Engine/Source/Runtime/Experimental/Chaos/Public/Chaos/Deformable/GaussSeidelCorotatedConstraints.h b/Engine/Source/Runtime/Experimental/Chaos/Public/Chaos/Deformable/GaussSeidelCorotatedConstraints.h
--- a/Engine/Source/Runtime/Experimental/Chaos/Public/Chaos/Deformable/GaussSeidelCorotatedConstraints.h
+++ b/Engine/Source/Runtime/Experimental/Chaos/Public/Chaos/Deformable/GaussSeidelCorotatedConstraints.h
@@ -58,7 +58,7 @@ namespace Chaos::Softs
 			{
 				for (int32 r = 0; r < 3; r++) {
 					for (int32 c = 0; c < 3; c++) {
-						DmInverse[(3 * 3) * e + 3 * r + c] *= Base::AlphaJArray[e];
+						DmInverse[(3 * 3) * e + 3 * r + c] /= Base::AlphaJArray[e];
 					}
 				}
 			}
