Add GeometryCollection category to flesh asset. GeometryCollection nodes (e.g. SkeletalMeshToCollection) are frequently used in flesh.

diff --git a/Engine/Plugins/Experimental/ChaosFlesh/Source/ChaosFleshNodes/Private/Dataflow/ChaosFleshNodesPlugin.cpp b/Engine/Plugins/Experimental/ChaosFlesh/Source/ChaosFleshNodes/Private/Dataflow/ChaosFleshNodesPlugin.cpp
--- a/Engine/Plugins/Experimental/ChaosFlesh/Source/ChaosFleshNodes/Private/Dataflow/ChaosFleshNodesPlugin.cpp
+++ b/Engine/Plugins/Experimental/ChaosFlesh/Source/ChaosFleshNodes/Private/Dataflow/ChaosFleshNodesPlugin.cpp
@@ -87,6 +87,7 @@ void IChaosFleshNodesPlugin::StartupModule()
 
 	// register node category for flesh asset
 	UE_DATAFLOW_REGISTER_CATEGORY_FORASSET_TYPE("Flesh", UFleshAsset);
+	UE_DATAFLOW_REGISTER_CATEGORY_FORASSET_TYPE("GeometryCollection", UFleshAsset);
 }
 
 void IChaosFleshNodesPlugin::ShutdownModule()
