Chaos Flesh:
- Fixed bug in GenerateSurfaceBindings when SKM has multiple LODs.
- Supported LOD selection.
- Supported surface embedding onto geometries with no tetrahedral mesh, e.g. binding to a kinematic surface.

diff --git a/Engine/Plugins/Experimental/ChaosFlesh/Source/ChaosFleshNodes/Private/Dataflow/ChaosFleshGenerateSurfaceBindingsNode.cpp b/Engine/Plugins/Experimental/ChaosFlesh/Source/ChaosFleshNodes/Private/Dataflow/ChaosFleshGenerateSurfaceBindingsNode.cpp
--- a/Engine/Plugins/Experimental/ChaosFlesh/Source/ChaosFleshNodes/Private/Dataflow/ChaosFleshGenerateSurfaceBindingsNode.cpp
+++ b/Engine/Plugins/Experimental/ChaosFlesh/Source/ChaosFleshNodes/Private/Dataflow/ChaosFleshGenerateSurfaceBindingsNode.cpp
@@ -230,7 +230,7 @@ FGenerateSurfaceBindings::Evaluate(UE::Dataflow::FContext& Context, const FDataf
 
 		const bool UseSkeletalMesh = SkeletalMesh != nullptr;
 		const bool UseStaticMesh = StaticMesh != nullptr;
-		if ((UseSkeletalMesh || UseStaticMesh) &&
+		if (IsConnected(&Collection) && (UseSkeletalMesh || UseStaticMesh) &&
 			Tetrahedron && TetrahedronStart && TetrahedronCount &&
 			Triangle && FacesStart && FacesCount &&
 			Vertex)
@@ -334,7 +334,7 @@ FGenerateSurfaceBindings::Evaluate(UE::Dataflow::FContext& Context, const FDataf
 			// Init boundary mesh for projections.
 			//
 			TArray<FIntVector> Triangles;
-			Chaos::FTriangleMesh TetBoundaryMesh;
+			Chaos::FTriangleMesh SurfaceMesh;
 			Chaos::FTriangleMesh::TBVHType<Chaos::FRealDouble> TetBoundaryBVH;
 			TArray<Chaos::FVec3> VertexD;
 			TConstArrayView<Chaos::TVec3<Chaos::FRealDouble>> VertexDView(VertexD);
@@ -361,20 +361,20 @@ FGenerateSurfaceBindings::Evaluate(UE::Dataflow::FContext& Context, const FDataf
 			}
 			else
 			{
-				for (int32 TetMeshIdx = 0; TetMeshIdx < TetrahedronStart->Num(); ++TetMeshIdx)
+				for (int32 GeometryIdx = 0; GeometryIdx < TetrahedronStart->Num(); ++GeometryIdx)
 				{
-					GeometryGroupSelected.Add(TetMeshIdx);
+					GeometryGroupSelected.Add(GeometryIdx);
 				}
 			}
 			TArray<int32> UsedTetsIndexToGlobalTetIndex;
 			TArray<int32> GlobalTetIndexToUsedTetsIndex;
 			GlobalTetIndexToUsedTetsIndex.Init(INDEX_NONE, Tetrahedron->Num());
 			int32 NumGuidHit = 0;
-			for (const int32& TetMeshIdx: GeometryGroupSelected)
+			for (const int32& GeometryIdx: GeometryGroupSelected)
 			{
 				if (GeometryGroupGuidsLocal.Num() && Guids)
 				{
-					if (!(*Guids)[TetMeshIdx].IsEmpty() && !GeometryGroupGuidsLocal.Contains((*Guids)[TetMeshIdx]))
+					if (!(*Guids)[GeometryIdx].IsEmpty() && !GeometryGroupGuidsLocal.Contains((*Guids)[GeometryIdx]))
 					{
 						continue;
 					}
@@ -383,29 +383,28 @@ FGenerateSurfaceBindings::Evaluate(UE::Dataflow::FContext& Context, const FDataf
 						NumGuidHit++;
 					}
 				}
-				const int32 TetMeshStart = (*TetrahedronStart)[TetMeshIdx];
-				const int32 TetMeshCount = (*TetrahedronCount)[TetMeshIdx];
-				if (TetMeshCount > 0)
+				const int32 TetMeshStart = (*TetrahedronStart)[GeometryIdx];
+				const int32 TetMeshCount = (*TetrahedronCount)[GeometryIdx];
+				Tets.Reserve(Tets.Num() + TetMeshCount);
+				UsedTetsIndexToGlobalTetIndex.Reserve(UsedTetsIndexToGlobalTetIndex.Num() + TetMeshCount);
+				for (int32 i = 0; i < TetMeshCount; i++)
 				{
-					for (int32 i = 0; i < TetMeshCount; i++)
-					{
-						const int32 Idx = TetMeshStart + i;
-						const FIntVector4& Tet = (*Tetrahedron)[Idx];
-						GlobalTetIndexToUsedTetsIndex[Idx] = Tets.Num();
-						Tets.Add(Chaos::TTetrahedron<Chaos::FReal>(
-							(*Vertex)[Tet[0]],
-							(*Vertex)[Tet[1]],
-							(*Vertex)[Tet[2]],
-							(*Vertex)[Tet[3]]));
-						UsedTetsIndexToGlobalTetIndex.Add(Idx);
-					}
-
-					if (bDoSurfaceProjection)
+					const int32 Idx = TetMeshStart + i;
+					const FIntVector4& Tet = (*Tetrahedron)[Idx];
+					GlobalTetIndexToUsedTetsIndex[Idx] = Tets.Num();
+					Tets.Add(Chaos::TTetrahedron<Chaos::FReal>(
+						(*Vertex)[Tet[0]],
+						(*Vertex)[Tet[1]],
+						(*Vertex)[Tet[2]],
+						(*Vertex)[Tet[3]]));
+					UsedTetsIndexToGlobalTetIndex.Add(Idx);
+				}
+				if (bDoSurfaceProjection)
+				{
+					Triangles.Reserve(Triangles.Num() + (*FacesCount)[GeometryIdx]);
+					for (int32 FaceIdx = (*FacesStart)[GeometryIdx]; FaceIdx < (*FacesStart)[GeometryIdx] + (*FacesCount)[GeometryIdx]; ++FaceIdx)
 					{
-						for (int32 FaceIdx = (*FacesStart)[TetMeshIdx]; FaceIdx < (*FacesStart)[TetMeshIdx] + (*FacesCount)[TetMeshIdx]; ++FaceIdx)
-						{
-							Triangles.Add((*Triangle)[FaceIdx]);
-						}
+						Triangles.Add((*Triangle)[FaceIdx]);
 					}
 				}
 			}
@@ -432,7 +431,7 @@ FGenerateSurfaceBindings::Evaluate(UE::Dataflow::FContext& Context, const FDataf
 			if (bDoSurfaceProjection)
 			{
 				//Todo(Chaosflesh): refactor reinterpret_cast in case the memory layout of the vector types not the same
-				TetBoundaryMesh.Init(
+				SurfaceMesh.Init(
 					reinterpret_cast<const TArray<Chaos::TVec3<int32>>&>(Triangles), 0, -1, false);
 
 				// Promote vertices to double because that's what FTriangleMesh wants.
@@ -443,11 +442,11 @@ FGenerateSurfaceBindings::Evaluate(UE::Dataflow::FContext& Context, const FDataf
 				}
 				VertexDView = TConstArrayView<Chaos::TVec3<Chaos::FRealDouble>>(VertexD);
 
-				PointNormals = TetBoundaryMesh.GetPointNormals(VertexDView, false, true);
-				TetBoundaryMesh.BuildBVH(VertexDView, TetBoundaryBVH);
+				PointNormals = SurfaceMesh.GetPointNormals(VertexDView, false, true);
+				SurfaceMesh.BuildBVH(VertexDView, TetBoundaryBVH);
 
 
-				TetBoundaryMesh.BuildSpatialHash(VertexDView, SpatialHash, Chaos::FRealDouble(SurfaceProjectionSearchRadius));
+				SurfaceMesh.BuildSpatialHash(VertexDView, SpatialHash, Chaos::FRealDouble(SurfaceProjectionSearchRadius));
 			}
 
 			//
@@ -458,28 +457,35 @@ FGenerateSurfaceBindings::Evaluate(UE::Dataflow::FContext& Context, const FDataf
 			TArray<TArray<FVector4f>> Weights;	 Weights.SetNum(RenderMeshVertices.Num());
 			TArray<TArray<FVector3f>> Offsets;	 Offsets.SetNum(RenderMeshVertices.Num());
 			TArray<TArray<float>> Masks;		 Masks.SetNum(RenderMeshVertices.Num());
-			TArray<int32> Orphans;
-			TArray<int32> OrphansCopy;
-			int32 TetHits = 0;
-			int32 TriHits = 0;
-			int32 Adoptions = 0;
-			int32 NumOrphans = 0;
+			TArray<int32> LOD0Orphans;
+
 			for (int32 LOD = 0; LOD < RenderMeshVertices.Num(); LOD++)
 			{
 				Parents[LOD].SetNumUninitialized(RenderMeshVertices[LOD].Num());
 				Weights[LOD].SetNumUninitialized(RenderMeshVertices[LOD].Num());
 				Offsets[LOD].SetNumUninitialized(RenderMeshVertices[LOD].Num());
 				Masks[LOD].SetNumUninitialized(RenderMeshVertices[LOD].Num());
-				int32 NumTetNotCollocated = 0;
-
-				TArray<int32> TetIntersections; TetIntersections.Reserve(64);
 				for (int32 i = 0; i < RenderMeshVertices[LOD].Num(); i++)
 				{
 					Parents[LOD][i] = FIntVector4(INDEX_NONE);
 					Weights[LOD][i] = FVector4f(0);
 					Offsets[LOD][i] = FVector3f(0);
 					Masks[LOD][i] = 0.0; // Shader does skinning for this vertex
+				}
+				if (SkeletalMeshLODList.Num() && !SkeletalMeshLODList.Contains(LOD))
+				{
+					continue;
+				}
+				TArray<int32> Orphans;
+				int32 TetHits = 0;
+				int32 TriHits = 0;
+				int32 Adoptions = 0;
+				int32 NumOrphans = 0;
+				int32 NumTetNotCollocated = 0;
 
+				TArray<int32> TetIntersections; TetIntersections.Reserve(64);
+				for (int32 i = 0; i < RenderMeshVertices[LOD].Num(); i++)
+				{
 					const FVector3f& Pos = RenderMeshVertices[LOD][i];
 					Chaos::TVec3<Chaos::FReal> PosD(Pos[0], Pos[1], Pos[2]);
 					TetIntersections = TetBVH.FindAllIntersections(PosD);
@@ -527,7 +533,7 @@ FGenerateSurfaceBindings::Evaluate(UE::Dataflow::FContext& Context, const FDataf
 						{
 							TArray<Chaos::TTriangleCollisionPoint<Chaos::FRealDouble>> Result;
 							//PointClosestTriangleQuery instead of SmoothProject
-							if (TetBoundaryMesh.PointClosestTriangleQuery(SpatialHash, VertexDView, i, TVec3Vert(Pos), Chaos::FRealDouble(SurfaceProjectionSearchRadius), Chaos::FRealDouble(SurfaceProjectionSearchRadius),
+							if (SurfaceMesh.PointClosestTriangleQuery(SpatialHash, VertexDView, i, TVec3Vert(Pos), Chaos::FRealDouble(SurfaceProjectionSearchRadius), Chaos::FRealDouble(SurfaceProjectionSearchRadius),
 								[](const int32 PointIndex, const int32 TriangleIndex)->bool
 								{
 									// use all nearby triangles
@@ -592,7 +598,10 @@ FGenerateSurfaceBindings::Evaluate(UE::Dataflow::FContext& Context, const FDataf
 					continue;
 				}
 				const TArray<TArray<uint32>>& NeighborNodes = RenderMeshNeighborNodes[LOD];
-				OrphansCopy = Orphans;
+				if (LOD == 0)
+				{
+					LOD0Orphans = Orphans;
+				}
 				TBitArray<> IsOrphan(false, RenderMeshVertices[LOD].Num());
 				while (bDoOrphanReparenting && Orphans.Num())
 				{
@@ -728,6 +737,10 @@ FGenerateSurfaceBindings::Evaluate(UE::Dataflow::FContext& Context, const FDataf
 						bUseSkeletalMeshImportModel ? TEXT("import") : TEXT("render"),
 						*MeshId, LOD,
 						RenderMeshVertices[LOD].Num(), TetHits, TriHits, Adoptions, NumOrphans);
+					Context.Warning(FString::Printf(
+						TEXT("GenerateSurfaceBindings Node: There are %d orphans."),
+						Orphans.Num()),
+						this, Out);
 				}
 				else
 				{
@@ -769,7 +782,7 @@ FGenerateSurfaceBindings::Evaluate(UE::Dataflow::FContext& Context, const FDataf
 				OutSKMDynamicMesh3.EnableAttributes();
 				OutSKMDynamicMesh3.EnableVertexColors(FVector3f(1, 0, 0));
 				TBitArray<> WasOrphan(false, RenderMeshVertices[0].Num());
-				for (int32 OrphanIdx : OrphansCopy)
+				for (int32 OrphanIdx : LOD0Orphans)
 				{
 					WasOrphan[OrphanIdx] = true;
 				}
diff --git a/Engine/Plugins/Experimental/ChaosFlesh/Source/ChaosFleshNodes/Public/Dataflow/ChaosFleshGenerateSurfaceBindingsNode.h b/Engine/Plugins/Experimental/ChaosFlesh/Source/ChaosFleshNodes/Public/Dataflow/ChaosFleshGenerateSurfaceBindingsNode.h
--- a/Engine/Plugins/Experimental/ChaosFlesh/Source/ChaosFleshNodes/Public/Dataflow/ChaosFleshGenerateSurfaceBindingsNode.h
+++ b/Engine/Plugins/Experimental/ChaosFlesh/Source/ChaosFleshNodes/Public/Dataflow/ChaosFleshGenerateSurfaceBindingsNode.h
@@ -16,9 +16,9 @@ class UDynamicMesh;
 
 DECLARE_LOG_CATEGORY_EXTERN(LogMeshBindings, Verbose, All);
 
-/* Generate barycentric bindings (used by the FleshDeformer deformer graph and Geometry Cache generation) of a render surface to a tetrahedral mesh.
+/* Generate barycentric bindings (used by the FleshDeformer deformer graph and Geometry Cache generation) of a render surface to a tetrahedral mesh and its surface.
 * If a point is outside of the tetrahedral mesh, find surface embedding within SurfaceProjectionSearchRadius.
-Embeddings are color coded in the render view:
+Embeddings of LOD 0 are color coded in the render view:
 	green: embedded on in a tetrahedron
 	blue: embedded on a surface triangle
 	red: orphan (cannot be embedded)
@@ -50,40 +50,43 @@ public:
 	}
 
 private:
-	/** Passthrough geometry collection. Bindings are stored as standalone groups in the \p Collection, keyed by the name of the input render mesh and all available LOD's. */
+	/** Collection containing tetrahedral mesh and surface mesh. Bindings are stored as standalone groups in the \p Collection, keyed by the name of the input render mesh and all available LOD's. */
 	UPROPERTY(meta = (DataflowInput, DataflowOutput, DataflowPassthrough = "Collection", DisplayName = "Collection"))
 	FManagedArrayCollection Collection;
 
 	/** The input mesh, whose render surface is used to generate bindings. */
-	UPROPERTY(EditAnywhere, Category = "Dataflow", DisplayName = "Static Mesh", meta = (DataflowInput, EditCondition = "SkeletalMeshIn == nullptr"))
+	UPROPERTY(EditAnywhere, Category = "Render Mesh", DisplayName = "Static Mesh", meta = (DataflowInput, EditCondition = "SkeletalMeshIn == nullptr"))
 	TObjectPtr<const UStaticMesh> StaticMeshIn = nullptr;
 
 	/** The input mesh, whose render surface is used to generate bindings. */
-	UPROPERTY(EditAnywhere, Category = "Dataflow", DisplayName = "Skeletal Mesh", meta = (DataflowInput, EditCondition = "StaticMeshIn == nullptr"))
+	UPROPERTY(EditAnywhere, Category = "Render Mesh", DisplayName = "Skeletal Mesh", meta = (DataflowInput, EditCondition = "StaticMeshIn == nullptr"))
 	TObjectPtr<const USkeletalMesh> SkeletalMeshIn = nullptr;
 
-	/** Only meshes with transforms in TransformSelection will be bound. */
+	/** Render mesh will only bind to geometries whose transforms are in TransformSelection. */
 	UPROPERTY(meta = (DataflowInput, DisplayName = "(Optional) TransformSelection"))
 	FDataflowTransformSelection TransformSelection;
 
-	/** Only meshes with transforms in TransformSelection will be bound. */
+	/** Render mesh will only bind to geometries whose GeometryGroupGuids match here. */
 	UPROPERTY(meta = (DataflowInput, DisplayName = "(Optional) GeometryGroupGuids"))
 	TArray<FString> GeometryGroupGuidsIn;
 
 	/** Use the import geometry of the skeletal mesh. */
-	UPROPERTY(EditAnywhere, Category = "Dataflow", meta = (DisplayName = "UseSkeletalMeshImportModel", EditCondition = "SkeletalMeshIn != nullptr"))
 	bool bUseSkeletalMeshImportModel = false;
 
+	/* Select skeletal mesh LODs to embed. Default empty list selects all LODs. */
+	UPROPERTY(EditAnywhere, Category = "Render Mesh", meta = (DisplayName = "Skeletal Mesh LOD List", EditCondition = "SkeletalMeshIn != nullptr && !bUseSkeletalMeshImportModel"))
+	TArray<int32> SkeletalMeshLODList;
+
 	/** Enable binding to the exterior hull of the tetrahedron mesh. */
-	UPROPERTY(EditAnywhere, Category = "Dataflow", meta = (DisplayName = "DoSurfaceProjection"))
+	UPROPERTY(EditAnywhere, Category = "Surface Projection", meta = (DisplayName = "DoSurfaceProjection"))
 	bool bDoSurfaceProjection = true;
 
 	/** The search radius when looking for surface triangles to bind to. */
-	UPROPERTY(EditAnywhere, Category = "Dataflow", meta = (DisplayName = "SurfaceProjectionSearchRadius", EditCondition = "bDoSurfaceProjection == true"))
+	UPROPERTY(EditAnywhere, Category = "Surface Projection", meta = (DisplayName = "SurfaceProjectionSearchRadius", EditCondition = "bDoSurfaceProjection == true"))
 	float SurfaceProjectionSearchRadius = 1.f;
 
 	/** When nodes aren't contained in tetrahedra and surface projection fails, try to find suitable bindings by looking to neighboring parents. */
-	UPROPERTY(EditAnywhere, Category = "Dataflow", meta = (DisplayName = "DoOrphanReparenting"))
+	UPROPERTY(EditAnywhere, Category = "Orphan Reparenting", meta = (DisplayName = "DoOrphanReparenting"))
 	bool bDoOrphanReparenting = true;
 
 	/** Converted from embedded skeletal/static mesh */
