Fix crash in CurveSamplingAnimationAssetTerminal when skm doesn't have skeleton

diff --git a/Engine/Plugins/Experimental/ChaosFlesh/Source/ChaosFleshNodes/Private/Dataflow/ChaosFleshFleshAssetTerminalNode.cpp b/Engine/Plugins/Experimental/ChaosFlesh/Source/ChaosFleshNodes/Private/Dataflow/ChaosFleshFleshAssetTerminalNode.cpp
--- a/Engine/Plugins/Experimental/ChaosFlesh/Source/ChaosFleshNodes/Private/Dataflow/ChaosFleshFleshAssetTerminalNode.cpp
+++ b/Engine/Plugins/Experimental/ChaosFlesh/Source/ChaosFleshNodes/Private/Dataflow/ChaosFleshFleshAssetTerminalNode.cpp
@@ -47,6 +47,14 @@ void FCurveSamplingAnimationAssetTerminalNode::SetAssetValue(TObjectPtr<UObject>
 #if WITH_EDITOR
 		if (AssetToSet)
 		{
+			if (!InSkeletalMesh->GetSkeleton())
+			{
+				Context.Error(FString::Printf(
+					TEXT("Input skeletal Mesh [%s] has no skeleton."),
+					*InSkeletalMesh->GetName()),
+					this);
+				return;
+			}
 			static bool bTransactRecording = false;
 
 			// set skeleton
