Fix crash in CreateAirVolumeConstraint when evaluating output

diff --git a/Engine/Plugins/Experimental/ChaosFlesh/Source/ChaosFleshNodes/Private/Dataflow/ChaosFleshSetVertexTrianglePositionTargetBindingNode.cpp b/Engine/Plugins/Experimental/ChaosFlesh/Source/ChaosFleshNodes/Private/Dataflow/ChaosFleshSetVertexTrianglePositionTargetBindingNode.cpp
--- a/Engine/Plugins/Experimental/ChaosFlesh/Source/ChaosFleshNodes/Private/Dataflow/ChaosFleshSetVertexTrianglePositionTargetBindingNode.cpp
+++ b/Engine/Plugins/Experimental/ChaosFlesh/Source/ChaosFleshNodes/Private/Dataflow/ChaosFleshSetVertexTrianglePositionTargetBindingNode.cpp
@@ -242,7 +242,7 @@ void FSetCollidableVerticesDataflowNode::Evaluate(UE::Dataflow::FContext& Contex
 void FCreateAirTetrahedralConstraintDataflowNode::Evaluate(UE::Dataflow::FContext& Context, const FDataflowOutput* Out) const
 {
 	using namespace UE::Geometry;
-	if (Out->IsA(&Collection))
+	if (Out->IsA(&Collection) || Out->IsA(&DynamicMesh))
 	{
 		TUniquePtr<FFleshCollection> InFleshCollection(GetValue(Context, &Collection).NewCopy<FFleshCollection>());
 		Chaos::FFleshCollectionFacade TetCollection(*InFleshCollection);
@@ -377,7 +377,7 @@ void FCreateAirTetrahedralConstraintDataflowNode::Evaluate(UE::Dataflow::FContex
 void FCreateAirVolumeConstraintDataflowNode::Evaluate(UE::Dataflow::FContext& Context, const FDataflowOutput* Out) const
 {
 	using namespace UE::Geometry;
-	if (Out->IsA(&Collection))
+	if (Out->IsA(&Collection) || Out->IsA(&DynamicMesh))
 	{
 		TUniquePtr<FFleshCollection> InFleshCollection(GetValue(Context, &Collection).NewCopy<FFleshCollection>());
 		Chaos::FFleshCollectionFacade TetCollection(*InFleshCollection);
