ChaosFlesh: fixed AppendTetrahedralCollection rendering bug

diff --git a/Engine/Plugins/Experimental/ChaosFlesh/Source/ChaosFleshNodes/Private/Dataflow/ChaosFleshAppendTetrahedralCollectionNode.cpp b/Engine/Plugins/Experimental/ChaosFlesh/Source/ChaosFleshNodes/Private/Dataflow/ChaosFleshAppendTetrahedralCollectionNode.cpp
--- a/Engine/Plugins/Experimental/ChaosFlesh/Source/ChaosFleshNodes/Private/Dataflow/ChaosFleshAppendTetrahedralCollectionNode.cpp
+++ b/Engine/Plugins/Experimental/ChaosFlesh/Source/ChaosFleshNodes/Private/Dataflow/ChaosFleshAppendTetrahedralCollectionNode.cpp
@@ -6,6 +6,86 @@
 #include "Dataflow/DataflowNodeFactory.h"
 #include "GeometryCollection/Facades/CollectionTransformFacade.h"
 
+
+FAppendTetrahedralCollectionDataflowNode_v2::FAppendTetrahedralCollectionDataflowNode_v2(const UE::Dataflow::FNodeParameters& InParam, FGuid InGuid) 
+	: FDataflowNode(InParam, InGuid)
+{
+	RegisterInputConnection(&Collection);
+	RegisterInputConnection(&CollectionToAppend);
+	RegisterOutputConnection(&Collection, &Collection);
+	RegisterOutputConnection(&CollectionGeometrySelection);
+	RegisterOutputConnection(&CollectionToAppendGeometrySelection);
+	RegisterOutputConnection(&CollectionGeometryGroupGuids);
+	RegisterOutputConnection(&CollectionToAppendGeometryGroupGuids);
+}
+
+void FAppendTetrahedralCollectionDataflowNode_v2::Evaluate(UE::Dataflow::FContext& Context, const FDataflowOutput* Out) const
+{
+	if (Out->IsA(&Collection))
+	{
+		TUniquePtr<FFleshCollection> InCollection(GetValue(Context, &Collection).NewCopy<FFleshCollection>());
+		TUniquePtr<FFleshCollection> InCollectionToAppend(GetValue(Context, &CollectionToAppend).NewCopy<FFleshCollection>());
+		TArray<FString> CollectionGeometryGroupGuidsLocal, CollectionToAppendGeometryGroupGuidsLocal;
+		if (const TManagedArray<FString>* CollectionGuidArray = InCollection->FindAttribute<FString>("Guid", FGeometryCollection::GeometryGroup))
+		{
+			CollectionGeometryGroupGuidsLocal = CollectionGuidArray->GetConstArray();
+		}
+		if (const TManagedArray<FString>* CollectionToAppendGuidArray = InCollectionToAppend->FindAttribute<FString>("Guid", FGeometryCollection::GeometryGroup))
+		{
+			CollectionToAppendGeometryGroupGuidsLocal = CollectionToAppendGuidArray->GetConstArray();
+		}
+		GeometryCollection::Facades::FCollectionTransformFacade InCollectionTransformFacade(*InCollection);
+		GeometryCollection::Facades::FCollectionTransformFacade InCollectionToAppendTransformFacade(*InCollectionToAppend);
+		const TMap<FString, int32> InCollectionBoneNameMap = InCollectionTransformFacade.BoneNameIndexMap();
+		const int32 CollectionNumTransform = InCollectionTransformFacade.Num();
+		const int32 CollectionToAppendNumTransform = InCollectionToAppendTransformFacade.Num();
+		//Append
+		InCollectionToAppend->AppendCollection(*InCollection); //InCollection is inserted to the front
+
+		if (bMergeTransform)
+		{
+			//Reorder and delete transform
+			if (const TManagedArray<FString>* InCollectionBoneNames = InCollectionToAppendTransformFacade.FindBoneNames())
+			{
+				TArray<int32> MergeRemapIndex;
+				TArray<int32> SortedMergeList;
+				for (int32 Idx = CollectionNumTransform; Idx < InCollectionBoneNames->Num(); ++Idx)
+				{
+					const FString& BoneName = (*InCollectionBoneNames)[Idx];
+					if (InCollectionBoneNameMap.Contains(BoneName) && !InCollectionToAppend->IsGeometry(Idx))
+					{
+						SortedMergeList.Add(Idx);
+						MergeRemapIndex.Add(InCollectionBoneNameMap[BoneName]);
+					}
+				}
+				InCollectionToAppend->MergeElements(FTransformCollection::TransformGroup, SortedMergeList, MergeRemapIndex);
+			}
+		}
+		const int32 CollectionNumGeometries = InCollection->NumElements(FGeometryCollection::GeometryGroup);
+		const int32 TotalNumGeometries = InCollectionToAppend->NumElements(FGeometryCollection::GeometryGroup);
+		FDataflowGeometrySelection InCollectionGeometrySelection;
+		FDataflowGeometrySelection InCollectionToAppendGeometrySelection;
+		InCollectionGeometrySelection.Initialize(TotalNumGeometries, false);
+		InCollectionToAppendGeometrySelection.Initialize(TotalNumGeometries, false);
+		for (int32 GeometryIdx = 0; GeometryIdx < TotalNumGeometries; ++GeometryIdx)
+		{
+			if (GeometryIdx < CollectionNumGeometries)
+			{
+				InCollectionGeometrySelection.SetSelected(GeometryIdx);
+			}
+			else
+			{
+				InCollectionToAppendGeometrySelection.SetSelected(GeometryIdx);
+			}
+		}
+		SetValue<const FManagedArrayCollection&>(Context, MoveTemp(*InCollectionToAppend), &Collection);
+		SetValue(Context, MoveTemp(InCollectionGeometrySelection), &CollectionGeometrySelection);
+		SetValue(Context, MoveTemp(InCollectionToAppendGeometrySelection), &CollectionToAppendGeometrySelection);
+		SetValue(Context, MoveTemp(CollectionGeometryGroupGuidsLocal), &CollectionGeometryGroupGuids);
+		SetValue(Context, MoveTemp(CollectionToAppendGeometryGroupGuidsLocal), &CollectionToAppendGeometryGroupGuids);
+	}
+}
+
 void FAppendTetrahedralCollectionDataflowNode::Evaluate(UE::Dataflow::FContext& Context, const FDataflowOutput* Out) const
 {
 	if (Out->IsA<FManagedArrayCollection>(&Collection1))
diff --git a/Engine/Plugins/Experimental/ChaosFlesh/Source/ChaosFleshNodes/Private/Dataflow/ChaosFleshNodesPlugin.cpp b/Engine/Plugins/Experimental/ChaosFlesh/Source/ChaosFleshNodes/Private/Dataflow/ChaosFleshNodesPlugin.cpp
--- a/Engine/Plugins/Experimental/ChaosFlesh/Source/ChaosFleshNodes/Private/Dataflow/ChaosFleshNodesPlugin.cpp
+++ b/Engine/Plugins/Experimental/ChaosFlesh/Source/ChaosFleshNodes/Private/Dataflow/ChaosFleshNodesPlugin.cpp
@@ -41,14 +41,13 @@
 void IChaosFleshNodesPlugin::StartupModule()
 {
 	DATAFLOW_NODE_REGISTER_CREATION_FACTORY(FAddKinematicParticlesDataflowNode);
-	DATAFLOW_NODE_REGISTER_CREATION_FACTORY(FAppendTetrahedralCollectionDataflowNode);
+	DATAFLOW_NODE_REGISTER_CREATION_FACTORY(FAppendTetrahedralCollectionDataflowNode_v2);
 	DATAFLOW_NODE_REGISTER_CREATION_FACTORY(FAppendToCollectionTransformAttributeDataflowNode);
 	DATAFLOW_NODE_REGISTER_CREATION_FACTORY(FAuthorSceneCollisionCandidates);
 	DATAFLOW_NODE_REGISTER_CREATION_FACTORY(FCalculateTetMetrics);
 	DATAFLOW_NODE_REGISTER_CREATION_FACTORY(FComputeFiberFieldNode);
 	DATAFLOW_NODE_REGISTER_CREATION_FACTORY(FComputeFiberStreamlineNode);
 	DATAFLOW_NODE_REGISTER_CREATION_FACTORY(FComputeIslandsNode);
-	DATAFLOW_NODE_REGISTER_CREATION_FACTORY(FComputeMuscleActivationDataNode);
 	DATAFLOW_NODE_REGISTER_CREATION_FACTORY(FComputeMuscleActivationDataNode_v2);
 	DATAFLOW_NODE_REGISTER_CREATION_FACTORY(FCreateAirTetrahedralConstraintDataflowNode);
 	DATAFLOW_NODE_REGISTER_CREATION_FACTORY(FCreateAirVolumeConstraintDataflowNode);
@@ -71,7 +70,6 @@ void IChaosFleshNodesPlugin::StartupModule()
 	DATAFLOW_NODE_REGISTER_CREATION_FACTORY(FRadialTetrahedronDataflowNodes);
 	DATAFLOW_NODE_REGISTER_CREATION_FACTORY(FReadSkeletalMeshCurvesDataflowNode);
 	DATAFLOW_NODE_REGISTER_CREATION_FACTORY(FSetCollidableVerticesDataflowNode);
-	DATAFLOW_NODE_REGISTER_CREATION_FACTORY(FSetFleshBonePositionTargetBindingDataflowNode);
 	DATAFLOW_NODE_REGISTER_CREATION_FACTORY(FSetFleshBonePositionTargetBindingDataflowNode_v2);
 	DATAFLOW_NODE_REGISTER_CREATION_FACTORY(FSetFleshDefaultPropertiesNode);
 	DATAFLOW_NODE_REGISTER_CREATION_FACTORY(FSetMuscleActivationParameterNode);
@@ -85,6 +83,11 @@ void IChaosFleshNodesPlugin::StartupModule()
 	DATAFLOW_NODE_REGISTER_CREATION_FACTORY(FVisualizePositionTargetsNode);
 	DATAFLOW_NODE_REGISTER_CREATION_FACTORY(FVisualizeKinematicFacesNode);
 
+	// Deprecated nodes
+	DATAFLOW_NODE_REGISTER_CREATION_FACTORY(FAppendTetrahedralCollectionDataflowNode);
+	DATAFLOW_NODE_REGISTER_CREATION_FACTORY(FComputeMuscleActivationDataNode);
+	DATAFLOW_NODE_REGISTER_CREATION_FACTORY(FSetFleshBonePositionTargetBindingDataflowNode);
+
 	// register node category for flesh asset
 	UE_DATAFLOW_REGISTER_CATEGORY_FORASSET_TYPE("Flesh", UFleshAsset);
 	UE_DATAFLOW_REGISTER_CATEGORY_FORASSET_TYPE("GeometryCollection", UFleshAsset);
diff --git a/Engine/Plugins/Experimental/ChaosFlesh/Source/ChaosFleshNodes/Public/Dataflow/ChaosFleshAppendTetrahedralCollectionNode.h b/Engine/Plugins/Experimental/ChaosFlesh/Source/ChaosFleshNodes/Public/Dataflow/ChaosFleshAppendTetrahedralCollectionNode.h
--- a/Engine/Plugins/Experimental/ChaosFlesh/Source/ChaosFleshNodes/Public/Dataflow/ChaosFleshAppendTetrahedralCollectionNode.h
+++ b/Engine/Plugins/Experimental/ChaosFlesh/Source/ChaosFleshNodes/Public/Dataflow/ChaosFleshAppendTetrahedralCollectionNode.h
@@ -14,16 +14,52 @@
 #include "ChaosFleshAppendTetrahedralCollectionNode.generated.h"
 
 /**
- * Append another Tetrahedral Collection to this collection. All attributes will be copied.
+ * Append another Tetrahedral Collection to this collection. All attributes will be appended to the end.
  */
 USTRUCT(meta = (DataflowFlesh))
+struct FAppendTetrahedralCollectionDataflowNode_v2 : public FDataflowNode
+{
+	GENERATED_USTRUCT_BODY()
+	DATAFLOW_NODE_DEFINE_INTERNAL(FAppendTetrahedralCollectionDataflowNode_v2, "AppendTetrahedralCollection", "Flesh", "")
+public:
+	FAppendTetrahedralCollectionDataflowNode_v2(const UE::Dataflow::FNodeParameters& InParam, FGuid InGuid = FGuid::NewGuid());
+
+private:
+	UPROPERTY(meta = (DataflowInput, DataflowOutput, DataflowPassthrough = "Collection", DataflowRenderGroups = "Surface"));
+	FManagedArrayCollection Collection;
+
+	UPROPERTY(meta = (DataflowInput));
+	FManagedArrayCollection CollectionToAppend;
+
+	/* If checked, non-geometry transforms from CollectionToAppend are merged into Collection by matching transform names. */
+	UPROPERTY(EditAnywhere, Category = "Dataflow", meta = (DisplayName = "Merge Transform"));
+	bool bMergeTransform = false;
+
+	UPROPERTY(meta = (DataflowOutput))
+	FDataflowGeometrySelection CollectionGeometrySelection;
+
+	UPROPERTY(meta = (DataflowOutput))
+	FDataflowGeometrySelection CollectionToAppendGeometrySelection;
+
+	UPROPERTY(meta = (DataflowOutput))
+	TArray<FString> CollectionGeometryGroupGuids;
+
+	UPROPERTY(meta = (DataflowOutput))
+	TArray<FString> CollectionToAppendGeometryGroupGuids;
+
+	virtual void Evaluate(UE::Dataflow::FContext& Context, const FDataflowOutput* Out) const override;
+};
+
+/**
+ * Append another Tetrahedral Collection to this collection. All attributes will be copied.
+ */
+USTRUCT(meta = (Deprecated = "5.7"))
 struct FAppendTetrahedralCollectionDataflowNode : public FDataflowNode
 {
 	GENERATED_USTRUCT_BODY()
 	DATAFLOW_NODE_DEFINE_INTERNAL(FAppendTetrahedralCollectionDataflowNode, "AppendTetrahedralCollection", "Flesh", "")
-	DATAFLOW_NODE_RENDER_TYPE("SurfaceRender", FGeometryCollection::StaticType(), "Collection1")
 public:
-	UPROPERTY(meta = (DataflowInput, DataflowOutput, DisplayName = "Collection1", DataflowPassthrough = "Collection1"));
+	UPROPERTY(meta = (DataflowInput, DataflowOutput, DisplayName = "Collection1", DataflowPassthrough = "Collection1", DataflowRenderGroups = "Surface"));
 	FManagedArrayCollection Collection1;
 
 	UPROPERTY(meta = (DataflowInput, DisplayName = "Collection2"));
