ChaosFlesh: uncomment propagate component function

diff --git a/Engine/Plugins/Experimental/ChaosFlesh/Source/ChaosFleshEngine/Private/ChaosFlesh/FleshAsset.cpp b/Engine/Plugins/Experimental/ChaosFlesh/Source/ChaosFleshEngine/Private/ChaosFlesh/FleshAsset.cpp
--- a/Engine/Plugins/Experimental/ChaosFlesh/Source/ChaosFleshEngine/Private/ChaosFlesh/FleshAsset.cpp
+++ b/Engine/Plugins/Experimental/ChaosFlesh/Source/ChaosFleshEngine/Private/ChaosFlesh/FleshAsset.cpp
@@ -10,8 +10,8 @@
 #include "Engine/SkeletalMesh.h"
 #include "GeometryCollection/TransformCollection.h"
 #if WITH_EDITOR
-//#include "ChaosFlesh/FleshComponent.h"
-//#include "UObject/UObjectIterator.h"
+#include "ChaosFlesh/FleshComponent.h"
+#include "UObject/UObjectIterator.h"
 #endif
 
 #include UE_INLINE_GENERATED_CPP_BY_NAME(FleshAsset)
@@ -167,15 +167,15 @@ void UFleshAsset::PostEditUndo()
 
 void UFleshAsset::PropagateTransformUpdateToComponents() const
 {
-	//for (TObjectIterator<UFleshComponent> It(RF_ClassDefaultObject, false, EInternalObjectFlags::Garbage); It; ++It)
-	//{
-	//	if (It->GetRestCollection() == this)
-	//	{
-	//		// make sure to reset the rest collection to make sure the internal state of the components is up to date 
-	//		// but we do not apply asset default to avoid overriding the existing overrides
-	//		It->SetRestCollection(this);
-	//	}
-	//}
+	for (TObjectIterator<UFleshComponent> It(RF_ClassDefaultObject, false, EInternalObjectFlags::Garbage); It; ++It)
+	{
+		if (It->GetRestCollection() == this)
+		{
+			// make sure to reset the rest collection to make sure the internal state of the components is up to date 
+			// but we do not apply asset default to avoid overriding the existing overrides
+			It->SetRestCollection(this);
+		}
+	}
 }
 #endif //if WITH_EDITOR
 
