USD: Centralized and exposed a single function to check if the USD SDK is enabled in UnrealUSDWrapper for the other modules to call.

#rnx

diff --git a/Engine/Plugins/Experimental/ChaosCachingUSD/Source/ChaosCachingUSD/ChaosCachingUSD.build.cs b/Engine/Plugins/Experimental/ChaosCachingUSD/Source/ChaosCachingUSD/ChaosCachingUSD.build.cs
--- a/Engine/Plugins/Experimental/ChaosCachingUSD/Source/ChaosCachingUSD/ChaosCachingUSD.build.cs
+++ b/Engine/Plugins/Experimental/ChaosCachingUSD/Source/ChaosCachingUSD/ChaosCachingUSD.build.cs
@@ -31,46 +31,7 @@ namespace UnrealBuildTool.Rules
 				}
 			);
 
-			bool bUseUsdMemoryOverloadsIfSupported = true;
-			bool bEnableUsdMemoryOverloads = (EnableUsdSdk(Target) && bUseUsdMemoryOverloadsIfSupported) &&
-				((Target.Platform == UnrealTargetPlatform.Win64 && Target.Architecture == UnrealArch.X64) //||
-				//Target.Platform == UnrealTargetPlatform.Mac ||
-				//Target.Platform == UnrealTargetPlatform.Linux
-				);
-			if (!bEnableUsdMemoryOverloads)
-			{
-				PrivateDefinitions.Add("SUPPRESS_PER_MODULE_INLINE_FILE"); // This module does not use core's standard operator new/delete overloads
-			}
-		}
-
-		bool EnableUsdSdk(ReadOnlyTargetRules Target)
-		{
-			bool bEnableUsdSdk = (
-				Target.WindowsPlatform.Compiler != WindowsCompiler.Intel &&
-				Target.StaticAnalyzer == StaticAnalyzer.None
-			);
-
-			// Don't enable USD when running the include tool because it has issues parsing Boost headers
-			if (Target.GlobalDefinitions.Contains("UE_INCLUDE_TOOL=1"))
-			{
-				bEnableUsdSdk = false;
-			}
-
-			if (!Target.bCompilePython)
-			{
-				bEnableUsdSdk = false;
-			}
-
-			bool bIsGameOrMonolithic = Target.LinkType == TargetLinkType.Monolithic || Target.Type == TargetType.Client || Target.Type == TargetType.Server || Target.Type == TargetType.Game;
-			if (bEnableUsdSdk &&
-				bIsGameOrMonolithic &&
-				((Target.StaticAllocator != StaticAllocatorType.Ansi) || !Target.GlobalDefinitions.Contains("UE_USE_MALLOC_FILL_BYTES=0"))
-			)
-			{
-				bEnableUsdSdk = false;
-			}
-
-			return bEnableUsdSdk;
+			UnrealUSDWrapper.CheckAndSetupUsdSdk(Target, this);
 		}
 	}
 }
\ No newline at end of file
diff --git a/Engine/Plugins/Experimental/ChaosFlesh/Source/ChaosFleshEditor/ChaosFleshEditor.Build.cs b/Engine/Plugins/Experimental/ChaosFlesh/Source/ChaosFleshEditor/ChaosFleshEditor.Build.cs
--- a/Engine/Plugins/Experimental/ChaosFlesh/Source/ChaosFleshEditor/ChaosFleshEditor.Build.cs
+++ b/Engine/Plugins/Experimental/ChaosFlesh/Source/ChaosFleshEditor/ChaosFleshEditor.Build.cs
@@ -85,9 +85,11 @@ namespace UnrealBuildTool.Rules
 
 			PrivateDefinitions.Add("CHAOS_INCLUDE_LEVEL_1=1");
 
+			bool bIsUsdSdkEnabled = UnrealUSDWrapper.CheckAndSetupUsdSdk(Target, this);
+
 			// Build flavor selection copied from UnrealUSDWrapper, then modified.
 			// Currently only Win64 is supported.
-			if (EnableUsdSdk(Target) && (Target.Type == TargetType.Editor && Target.Platform == UnrealTargetPlatform.Win64))
+			if (bIsUsdSdkEnabled && (Target.Type == TargetType.Editor && Target.Platform == UnrealTargetPlatform.Win64))
 			{
 				bUseRTTI = true;
 				PublicDefinitions.Add("DO_USD_CACHING=1");
@@ -97,62 +99,5 @@ namespace UnrealBuildTool.Rules
 				PublicDefinitions.Add("DO_USD_CACHING=0");
 			}
 		}
-		bool EnableUsdSdk(ReadOnlyTargetRules Target)
-		{
-			// USD SDK has been built against Python 3 and won't launch if the editor is using Python 2
-
-			bool bEnableUsdSdk = (
-				Target.WindowsPlatform.Compiler != WindowsCompiler.Clang &&
-				Target.StaticAnalyzer == StaticAnalyzer.None
-			);
-
-			// Don't enable USD when running the include tool because it has issues parsing Boost headers
-			if (Target.GlobalDefinitions.Contains("UE_INCLUDE_TOOL=1"))
-			{
-				bEnableUsdSdk = false;
-			}
-
-			if (!Target.bCompilePython)
-			{
-				bEnableUsdSdk = false;
-			}
-
-			// If you want to use USD in a monolithic target, you'll have to use the ANSI allocator.
-			// USD always uses the ANSI C allocators directly. In a DLL UE build (so not monolithic) we can just override the operators new and delete
-			// on each module with versions that use either the ANSI (so USD-compatible) allocators or the UE allocators (ModuleBoilerplate.h) when appropriate.
-			// In a monolithic build we can't do that, as the primary game module will already define overrides for operator new and delete with
-			// the standard UE allocators: Since we can only have one operator new/delete override on the entire monolithic executable, we can't define our own overrides.
-			// Additionally, the ANSI allocator does not work properly with FMallocPoisonProxy. Consequently, FMallocPoisonProxy has to be disabled.
-			// The only way around it is by forcing the ansi allocator and disabling FMallocPoisonProxy in your project's target file
-			// (YourProject/Source/YourProject.Target.cs) file like this:
-			//
-			//		public class YourProject : TargetRules
-			//		{
-			//			public YourProject(TargetInfo Target) : base(Target)
-			//			{
-			//				...
-			//				GlobalDefinitions.Add("FORCE_ANSI_ALLOCATOR=1");
-			//				GlobalDefinitions.Add("UE_USE_MALLOC_FILL_BYTES=0");
-			//				...
-			//			}
-			//		}
-			//
-			// This will force the entire built executable to use the ANSI C allocators for everything (by disabling the UE overrides in ModuleBoilerplate.h) while
-			// FMallocPoisonProxy is disabled, and so UE and USD allocations will be compatible.
-			// Note that by that point everything will be using the USD-compatible ANSI allocators anyway, so our overrides in USDMemory.h are also disabled, as they're unnecessary.
-			// Also note that we're forced to use dynamic linking for monolithic targets mainly because static linking the USD libraries disables support for user USD plugins, and secondly
-			// because those static libraries would need to be linked with the --whole-archive argument, and there is currently no standard way of doing that in UE.
-			if (bEnableUsdSdk && Target.LinkType == TargetLinkType.Monolithic && (Target.StaticAllocator != StaticAllocatorType.Ansi) && !Target.GlobalDefinitions.Contains("UE_USE_MALLOC_FILL_BYTES=0"))
-			{
-				//PublicDefinitions.Add("USD_FORCE_DISABLED=1");
-				bEnableUsdSdk = false;
-			}
-			else
-			{
-				//PublicDefinitions.Add("USD_FORCE_DISABLED=0");
-			}
-
-			return bEnableUsdSdk;
-		}
 	}
 }
diff --git a/Engine/Plugins/Experimental/ChaosFlesh/Source/ChaosFleshEngine/ChaosFleshEngine.Build.cs b/Engine/Plugins/Experimental/ChaosFlesh/Source/ChaosFleshEngine/ChaosFleshEngine.Build.cs
--- a/Engine/Plugins/Experimental/ChaosFlesh/Source/ChaosFleshEngine/ChaosFleshEngine.Build.cs
+++ b/Engine/Plugins/Experimental/ChaosFlesh/Source/ChaosFleshEngine/ChaosFleshEngine.Build.cs
@@ -81,85 +81,19 @@ namespace UnrealBuildTool.Rules
 
 			PrivateDefinitions.Add("CHAOS_INCLUDE_LEVEL_1=1");
 
+			bool bIsUsdSdkEnabled = UnrealUSDWrapper.CheckAndSetupUsdSdk(Target, this);
+
 			// Build flavor selection copied from UnrealUSDWrapper, then modified.
 			// Currently only Win64 is supported.
-			if (EnableUsdSdk(Target) && (Target.Type == TargetType.Editor && Target.Platform == UnrealTargetPlatform.Win64))
+			if (bIsUsdSdkEnabled && (Target.Type == TargetType.Editor && Target.Platform == UnrealTargetPlatform.Win64))
 			{
 				bUseRTTI = true;
 				PublicDefinitions.Add("DO_USD_CACHING=1");
-
-				bool bUseUsdMemoryOverloadsIfSupported = true;
-				bool bEnableUsdMemoryOverloads = (EnableUsdSdk(Target) && bUseUsdMemoryOverloadsIfSupported) &&
-					((Target.Platform == UnrealTargetPlatform.Win64 && Target.Architecture == UnrealArch.X64) //||
-					//Target.Platform == UnrealTargetPlatform.Mac ||
-					//Target.Platform == UnrealTargetPlatform.Linux
-					);
-				if (!bEnableUsdMemoryOverloads)
-				{
-					PrivateDefinitions.Add("SUPPRESS_PER_MODULE_INLINE_FILE"); // This module does not use core's standard operator new/delete overloads
-				}
 			}
 			else
 			{
 				PublicDefinitions.Add("DO_USD_CACHING=0");
 			}
 		}
-		bool EnableUsdSdk(ReadOnlyTargetRules Target)
-		{
-			// USD SDK has been built against Python 3 and won't launch if the editor is using Python 2
-
-			bool bEnableUsdSdk = (
-				Target.WindowsPlatform.Compiler != WindowsCompiler.Clang &&
-				Target.StaticAnalyzer == StaticAnalyzer.None
-			);
-
-			// Don't enable USD when running the include tool because it has issues parsing Boost headers
-			if (Target.GlobalDefinitions.Contains("UE_INCLUDE_TOOL=1"))
-			{
-				bEnableUsdSdk = false;
-			}
-
-			if (!Target.bCompilePython)
-			{
-				bEnableUsdSdk = false;
-			}
-
-			// If you want to use USD in a monolithic target, you'll have to use the ANSI allocator.
-			// USD always uses the ANSI C allocators directly. In a DLL UE build (so not monolithic) we can just override the operators new and delete
-			// on each module with versions that use either the ANSI (so USD-compatible) allocators or the UE allocators (ModuleBoilerplate.h) when appropriate.
-			// In a monolithic build we can't do that, as the primary game module will already define overrides for operator new and delete with
-			// the standard UE allocators: Since we can only have one operator new/delete override on the entire monolithic executable, we can't define our own overrides.
-			// Additionally, the ANSI allocator does not work properly with FMallocPoisonProxy. Consequently, FMallocPoisonProxy has to be disabled.
-			// The only way around it is by forcing the ansi allocator and disabling FMallocPoisonProxy in your project's target file
-			// (YourProject/Source/YourProject.Target.cs) file like this:
-			//
-			//		public class YourProject : TargetRules
-			//		{
-			//			public YourProject(TargetInfo Target) : base(Target)
-			//			{
-			//				...
-			//				GlobalDefinitions.Add("FORCE_ANSI_ALLOCATOR=1");
-			//				GlobalDefinitions.Add("UE_USE_MALLOC_FILL_BYTES=0");
-			//				...
-			//			}
-			//		}
-			//
-			// This will force the entire built executable to use the ANSI C allocators for everything (by disabling the UE overrides in ModuleBoilerplate.h) while
-			// FMallocPoisonProxy is disabled, and so UE and USD allocations will be compatible.
-			// Note that by that point everything will be using the USD-compatible ANSI allocators anyway, so our overrides in USDMemory.h are also disabled, as they're unnecessary.
-			// Also note that we're forced to use dynamic linking for monolithic targets mainly because static linking the USD libraries disables support for user USD plugins, and secondly
-			// because those static libraries would need to be linked with the --whole-archive argument, and there is currently no standard way of doing that in UE.
-			if (bEnableUsdSdk && Target.LinkType == TargetLinkType.Monolithic && (Target.StaticAllocator != StaticAllocatorType.Ansi) && !Target.GlobalDefinitions.Contains("UE_USE_MALLOC_FILL_BYTES=0"))
-			{
-				//PublicDefinitions.Add("USD_FORCE_DISABLED=1");
-				bEnableUsdSdk = false;
-			}
-			else
-			{
-				//PublicDefinitions.Add("USD_FORCE_DISABLED=0");
-			}
-
-			return bEnableUsdSdk;
-		}
 	}
 }
diff --git a/Engine/Plugins/Importers/USDImporter/Source/USDExporter/USDExporter.Build.cs b/Engine/Plugins/Importers/USDImporter/Source/USDExporter/USDExporter.Build.cs
--- a/Engine/Plugins/Importers/USDImporter/Source/USDExporter/USDExporter.Build.cs
+++ b/Engine/Plugins/Importers/USDImporter/Source/USDExporter/USDExporter.Build.cs
@@ -61,46 +61,7 @@ namespace UnrealBuildTool.Rules
 				}
 			);
 
-			bool bUseUsdMemoryOverloadsIfSupported = true;
-			bool bEnableUsdMemoryOverloads = (EnableUsdSdk(Target) && bUseUsdMemoryOverloadsIfSupported) &&
-				((Target.Platform == UnrealTargetPlatform.Win64 && Target.Architecture == UnrealArch.X64) //||
-				//Target.Platform == UnrealTargetPlatform.Mac ||
-				//Target.Platform == UnrealTargetPlatform.Linux
-				);
-			if (!bEnableUsdMemoryOverloads)
-			{
-				PrivateDefinitions.Add("SUPPRESS_PER_MODULE_INLINE_FILE"); // This module does not use core's standard operator new/delete overloads
-			}
-		}
-
-		bool EnableUsdSdk(ReadOnlyTargetRules Target)
-		{
-			bool bEnableUsdSdk = (
-				Target.WindowsPlatform.Compiler != WindowsCompiler.Intel &&
-				Target.StaticAnalyzer == StaticAnalyzer.None
-			);
-
-			// Don't enable USD when running the include tool because it has issues parsing Boost headers
-			if (Target.GlobalDefinitions.Contains("UE_INCLUDE_TOOL=1"))
-			{
-				bEnableUsdSdk = false;
-			}
-
-			if (!Target.bCompilePython)
-			{
-				bEnableUsdSdk = false;
-			}
-
-			bool bIsGameOrMonolithic = Target.LinkType == TargetLinkType.Monolithic || Target.Type == TargetType.Client || Target.Type == TargetType.Server || Target.Type == TargetType.Game;
-			if (bEnableUsdSdk &&
-				bIsGameOrMonolithic &&
-				((Target.StaticAllocator != StaticAllocatorType.Ansi) || !Target.GlobalDefinitions.Contains("UE_USE_MALLOC_FILL_BYTES=0"))
-			)
-			{
-				bEnableUsdSdk = false;
-			}
-
-			return bEnableUsdSdk;
+			UnrealUSDWrapper.CheckAndSetupUsdSdk(Target, this);
 		}
 	}
 }
diff --git a/Engine/Plugins/Importers/USDImporter/Source/USDSchemas/USDSchemas.Build.cs b/Engine/Plugins/Importers/USDImporter/Source/USDSchemas/USDSchemas.Build.cs
--- a/Engine/Plugins/Importers/USDImporter/Source/USDSchemas/USDSchemas.Build.cs
+++ b/Engine/Plugins/Importers/USDImporter/Source/USDSchemas/USDSchemas.Build.cs
@@ -82,46 +82,7 @@ namespace UnrealBuildTool.Rules
 				}
 			}
 
-			bool bUseUsdMemoryOverloadsIfSupported = true;
-			bool bEnableUsdMemoryOverloads = (EnableUsdSdk(Target) && bUseUsdMemoryOverloadsIfSupported) &&
-				((Target.Platform == UnrealTargetPlatform.Win64 && Target.Architecture == UnrealArch.X64) //||
-				//Target.Platform == UnrealTargetPlatform.Mac ||
-				//Target.Platform == UnrealTargetPlatform.Linux
-				);
-			if (!bEnableUsdMemoryOverloads)
-			{
-				PrivateDefinitions.Add("SUPPRESS_PER_MODULE_INLINE_FILE"); // This module does not use core's standard operator new/delete overloads
-			}
-		}
-
-		bool EnableUsdSdk(ReadOnlyTargetRules Target)
-		{
-			bool bEnableUsdSdk = (
-				Target.WindowsPlatform.Compiler != WindowsCompiler.Intel &&
-				Target.StaticAnalyzer == StaticAnalyzer.None
-			);
-
-			// Don't enable USD when running the include tool because it has issues parsing Boost headers
-			if (Target.GlobalDefinitions.Contains("UE_INCLUDE_TOOL=1"))
-			{
-				bEnableUsdSdk = false;
-			}
-
-			if (!Target.bCompilePython)
-			{
-				bEnableUsdSdk = false;
-			}
-
-			bool bIsGameOrMonolithic = Target.LinkType == TargetLinkType.Monolithic || Target.Type == TargetType.Client || Target.Type == TargetType.Server || Target.Type == TargetType.Game;
-			if (bEnableUsdSdk &&
-				bIsGameOrMonolithic &&
-				((Target.StaticAllocator != StaticAllocatorType.Ansi) || !Target.GlobalDefinitions.Contains("UE_USE_MALLOC_FILL_BYTES=0"))
-			)
-			{
-				bEnableUsdSdk = false;
-			}
-
-			return bEnableUsdSdk;
+			UnrealUSDWrapper.CheckAndSetupUsdSdk(Target, this);
 		}
 	}
 }
diff --git a/Engine/Plugins/Importers/USDImporter/Source/USDStage/USDStage.Build.cs b/Engine/Plugins/Importers/USDImporter/Source/USDStage/USDStage.Build.cs
--- a/Engine/Plugins/Importers/USDImporter/Source/USDStage/USDStage.Build.cs
+++ b/Engine/Plugins/Importers/USDImporter/Source/USDStage/USDStage.Build.cs
@@ -72,46 +72,7 @@ namespace UnrealBuildTool.Rules
 				);
 			}
 
-			bool bUseUsdMemoryOverloadsIfSupported = true;
-			bool bEnableUsdMemoryOverloads = (EnableUsdSdk(Target) && bUseUsdMemoryOverloadsIfSupported) &&
-				((Target.Platform == UnrealTargetPlatform.Win64 && Target.Architecture == UnrealArch.X64) //||
-				//Target.Platform == UnrealTargetPlatform.Mac ||
-				//Target.Platform == UnrealTargetPlatform.Linux
-				);
-			if (!bEnableUsdMemoryOverloads)
-			{
-				PrivateDefinitions.Add("SUPPRESS_PER_MODULE_INLINE_FILE"); // This module does not use core's standard operator new/delete overloads
-			}
-		}
-
-		bool EnableUsdSdk(ReadOnlyTargetRules Target)
-		{
-			bool bEnableUsdSdk = (
-				Target.WindowsPlatform.Compiler != WindowsCompiler.Intel &&
-				Target.StaticAnalyzer == StaticAnalyzer.None
-			);
-
-			// Don't enable USD when running the include tool because it has issues parsing Boost headers
-			if (Target.GlobalDefinitions.Contains("UE_INCLUDE_TOOL=1"))
-			{
-				bEnableUsdSdk = false;
-			}
-
-			if (!Target.bCompilePython)
-			{
-				bEnableUsdSdk = false;
-			}
-
-			bool bIsGameOrMonolithic = Target.LinkType == TargetLinkType.Monolithic || Target.Type == TargetType.Client || Target.Type == TargetType.Server || Target.Type == TargetType.Game;
-			if (bEnableUsdSdk &&
-				bIsGameOrMonolithic &&
-				((Target.StaticAllocator != StaticAllocatorType.Ansi) || !Target.GlobalDefinitions.Contains("UE_USE_MALLOC_FILL_BYTES=0"))
-			)
-			{
-				bEnableUsdSdk = false;
-			}
-
-			return bEnableUsdSdk;
+			UnrealUSDWrapper.CheckAndSetupUsdSdk(Target, this);
 		}
 	}
 }
diff --git a/Engine/Plugins/Importers/USDImporter/Source/USDStageEditor/USDStageEditor.Build.cs b/Engine/Plugins/Importers/USDImporter/Source/USDStageEditor/USDStageEditor.Build.cs
--- a/Engine/Plugins/Importers/USDImporter/Source/USDStageEditor/USDStageEditor.Build.cs
+++ b/Engine/Plugins/Importers/USDImporter/Source/USDStageEditor/USDStageEditor.Build.cs
@@ -40,46 +40,7 @@ namespace UnrealBuildTool.Rules
 				}
 			);
 
-			bool bUseUsdMemoryOverloadsIfSupported = true;
-			bool bEnableUsdMemoryOverloads = (EnableUsdSdk(Target) && bUseUsdMemoryOverloadsIfSupported) &&
-				((Target.Platform == UnrealTargetPlatform.Win64 && Target.Architecture == UnrealArch.X64) //||
-				//Target.Platform == UnrealTargetPlatform.Mac ||
-				//Target.Platform == UnrealTargetPlatform.Linux
-				);
-			if (!bEnableUsdMemoryOverloads)
-			{
-				PrivateDefinitions.Add("SUPPRESS_PER_MODULE_INLINE_FILE"); // This module does not use core's standard operator new/delete overloads
-			}
-		}
-
-		bool EnableUsdSdk(ReadOnlyTargetRules Target)
-		{
-			bool bEnableUsdSdk = (
-				Target.WindowsPlatform.Compiler != WindowsCompiler.Intel &&
-				Target.StaticAnalyzer == StaticAnalyzer.None
-			);
-
-			// Don't enable USD when running the include tool because it has issues parsing Boost headers
-			if (Target.GlobalDefinitions.Contains("UE_INCLUDE_TOOL=1"))
-			{
-				bEnableUsdSdk = false;
-			}
-
-			if (!Target.bCompilePython)
-			{
-				bEnableUsdSdk = false;
-			}
-
-			bool bIsGameOrMonolithic = Target.LinkType == TargetLinkType.Monolithic || Target.Type == TargetType.Client || Target.Type == TargetType.Server || Target.Type == TargetType.Game;
-			if (bEnableUsdSdk &&
-				bIsGameOrMonolithic &&
-				((Target.StaticAllocator != StaticAllocatorType.Ansi) || !Target.GlobalDefinitions.Contains("UE_USE_MALLOC_FILL_BYTES=0"))
-			)
-			{
-				bEnableUsdSdk = false;
-			}
-
-			return bEnableUsdSdk;
+			UnrealUSDWrapper.CheckAndSetupUsdSdk(Target, this);
 		}
 	}
 }
diff --git a/Engine/Plugins/Importers/USDImporter/Source/USDStageEditorViewModels/USDStageEditorViewModels.Build.cs b/Engine/Plugins/Importers/USDImporter/Source/USDStageEditorViewModels/USDStageEditorViewModels.Build.cs
--- a/Engine/Plugins/Importers/USDImporter/Source/USDStageEditorViewModels/USDStageEditorViewModels.Build.cs
+++ b/Engine/Plugins/Importers/USDImporter/Source/USDStageEditorViewModels/USDStageEditorViewModels.Build.cs
@@ -34,46 +34,7 @@ namespace UnrealBuildTool.Rules
 				}
 			);
 
-			bool bUseUsdMemoryOverloadsIfSupported = true;
-			bool bEnableUsdMemoryOverloads = (EnableUsdSdk(Target) && bUseUsdMemoryOverloadsIfSupported) &&
-				((Target.Platform == UnrealTargetPlatform.Win64 && Target.Architecture == UnrealArch.X64) //||
-				//Target.Platform == UnrealTargetPlatform.Mac ||
-				//Target.Platform == UnrealTargetPlatform.Linux
-				);
-			if (!bEnableUsdMemoryOverloads)
-			{
-				PrivateDefinitions.Add("SUPPRESS_PER_MODULE_INLINE_FILE"); // This module does not use core's standard operator new/delete overloads
-			}
-		}
-
-		bool EnableUsdSdk(ReadOnlyTargetRules Target)
-		{
-			bool bEnableUsdSdk = (
-				Target.WindowsPlatform.Compiler != WindowsCompiler.Intel &&
-				Target.StaticAnalyzer == StaticAnalyzer.None
-			);
-
-			// Don't enable USD when running the include tool because it has issues parsing Boost headers
-			if (Target.GlobalDefinitions.Contains("UE_INCLUDE_TOOL=1"))
-			{
-				bEnableUsdSdk = false;
-			}
-
-			if (!Target.bCompilePython)
-			{
-				bEnableUsdSdk = false;
-			}
-
-			bool bIsGameOrMonolithic = Target.LinkType == TargetLinkType.Monolithic || Target.Type == TargetType.Client || Target.Type == TargetType.Server || Target.Type == TargetType.Game;
-			if (bEnableUsdSdk &&
-				bIsGameOrMonolithic &&
-				((Target.StaticAllocator != StaticAllocatorType.Ansi) || !Target.GlobalDefinitions.Contains("UE_USE_MALLOC_FILL_BYTES=0"))
-			)
-			{
-				bEnableUsdSdk = false;
-			}
-
-			return bEnableUsdSdk;
+			UnrealUSDWrapper.CheckAndSetupUsdSdk(Target, this);
 		}
 	}
 }
diff --git a/Engine/Plugins/Importers/USDImporter/Source/USDStageImporter/USDStageImporter.Build.cs b/Engine/Plugins/Importers/USDImporter/Source/USDStageImporter/USDStageImporter.Build.cs
--- a/Engine/Plugins/Importers/USDImporter/Source/USDStageImporter/USDStageImporter.Build.cs
+++ b/Engine/Plugins/Importers/USDImporter/Source/USDStageImporter/USDStageImporter.Build.cs
@@ -46,46 +46,7 @@ namespace UnrealBuildTool.Rules
 				}
 			);
 
-			bool bUseUsdMemoryOverloadsIfSupported = true;
-			bool bEnableUsdMemoryOverloads = (EnableUsdSdk(Target) && bUseUsdMemoryOverloadsIfSupported) &&
-				((Target.Platform == UnrealTargetPlatform.Win64 && Target.Architecture == UnrealArch.X64) //||
-				//Target.Platform == UnrealTargetPlatform.Mac ||
-				//Target.Platform == UnrealTargetPlatform.Linux
-				);
-			if (!bEnableUsdMemoryOverloads)
-			{
-				PrivateDefinitions.Add("SUPPRESS_PER_MODULE_INLINE_FILE"); // This module does not use core's standard operator new/delete overloads
-			}
-		}
-
-		bool EnableUsdSdk(ReadOnlyTargetRules Target)
-		{
-			bool bEnableUsdSdk = (
-				Target.WindowsPlatform.Compiler != WindowsCompiler.Intel &&
-				Target.StaticAnalyzer == StaticAnalyzer.None
-			);
-
-			// Don't enable USD when running the include tool because it has issues parsing Boost headers
-			if (Target.GlobalDefinitions.Contains("UE_INCLUDE_TOOL=1"))
-			{
-				bEnableUsdSdk = false;
-			}
-
-			if (!Target.bCompilePython)
-			{
-				bEnableUsdSdk = false;
-			}
-
-			bool bIsGameOrMonolithic = Target.LinkType == TargetLinkType.Monolithic || Target.Type == TargetType.Client || Target.Type == TargetType.Server || Target.Type == TargetType.Game;
-			if (bEnableUsdSdk &&
-				bIsGameOrMonolithic &&
-				((Target.StaticAllocator != StaticAllocatorType.Ansi) || !Target.GlobalDefinitions.Contains("UE_USE_MALLOC_FILL_BYTES=0"))
-			)
-			{
-				bEnableUsdSdk = false;
-			}
-
-			return bEnableUsdSdk;
+			UnrealUSDWrapper.CheckAndSetupUsdSdk(Target, this);
 		}
 	}
 }
diff --git a/Engine/Plugins/Importers/USDImporter/Source/USDTests/USDTests.Build.cs b/Engine/Plugins/Importers/USDImporter/Source/USDTests/USDTests.Build.cs
--- a/Engine/Plugins/Importers/USDImporter/Source/USDTests/USDTests.Build.cs
+++ b/Engine/Plugins/Importers/USDImporter/Source/USDTests/USDTests.Build.cs
@@ -35,46 +35,7 @@ namespace UnrealBuildTool.Rules
 				);
 			}
 
-			bool bUseUsdMemoryOverloadsIfSupported = true;
-			bool bEnableUsdMemoryOverloads = (EnableUsdSdk(Target) && bUseUsdMemoryOverloadsIfSupported) &&
-				((Target.Platform == UnrealTargetPlatform.Win64 && Target.Architecture == UnrealArch.X64) //||
-				//Target.Platform == UnrealTargetPlatform.Mac ||
-				//Target.Platform == UnrealTargetPlatform.Linux
-				);
-			if (!bEnableUsdMemoryOverloads)
-			{
-				PrivateDefinitions.Add("SUPPRESS_PER_MODULE_INLINE_FILE"); // This module does not use core's standard operator new/delete overloads
-			}
-		}
-
-		bool EnableUsdSdk(ReadOnlyTargetRules Target)
-		{
-			bool bEnableUsdSdk = (
-				Target.WindowsPlatform.Compiler != WindowsCompiler.Intel &&
-				Target.StaticAnalyzer == StaticAnalyzer.None
-			);
-
-			// Don't enable USD when running the include tool because it has issues parsing Boost headers
-			if (Target.GlobalDefinitions.Contains("UE_INCLUDE_TOOL=1"))
-			{
-				bEnableUsdSdk = false;
-			}
-
-			if (!Target.bCompilePython)
-			{
-				bEnableUsdSdk = false;
-			}
-
-			bool bIsGameOrMonolithic = Target.LinkType == TargetLinkType.Monolithic || Target.Type == TargetType.Client || Target.Type == TargetType.Server || Target.Type == TargetType.Game;
-			if (bEnableUsdSdk &&
-				bIsGameOrMonolithic &&
-				((Target.StaticAllocator != StaticAllocatorType.Ansi) || !Target.GlobalDefinitions.Contains("UE_USE_MALLOC_FILL_BYTES=0"))
-			)
-			{
-				bEnableUsdSdk = false;
-			}
-
-			return bEnableUsdSdk;
+			UnrealUSDWrapper.CheckAndSetupUsdSdk(Target, this);
 		}
 	}
 }
diff --git a/Engine/Plugins/Importers/USDImporterMDL/Source/USDImporterMDL/USDImporterMDL.Build.cs b/Engine/Plugins/Importers/USDImporterMDL/Source/USDImporterMDL/USDImporterMDL.Build.cs
--- a/Engine/Plugins/Importers/USDImporterMDL/Source/USDImporterMDL/USDImporterMDL.Build.cs
+++ b/Engine/Plugins/Importers/USDImporterMDL/Source/USDImporterMDL/USDImporterMDL.Build.cs
@@ -4,6 +4,7 @@ using System;
 using System.IO;
 using System.Linq;
 using UnrealBuildTool;
+using UnrealBuildTool.Rules;
 
 public class USDImporterMDL : ModuleRules
 {
@@ -30,16 +31,7 @@ public class USDImporterMDL : ModuleRules
 			}
 		);
 
-		bool bUseUsdMemoryOverloadsIfSupported = true;
-		bool bEnableUsdMemoryOverloads = (EnableUsdSdk(Target) && bUseUsdMemoryOverloadsIfSupported) &&
-			((Target.Platform == UnrealTargetPlatform.Win64 && Target.Architecture == UnrealArch.X64) //||
-			//Target.Platform == UnrealTargetPlatform.Mac ||
-			//Target.Platform == UnrealTargetPlatform.Linux
-			);
-		if (!bEnableUsdMemoryOverloads)
-		{
-			PrivateDefinitions.Add("SUPPRESS_PER_MODULE_INLINE_FILE"); // This module does not use core's standard operator new/delete overloads
-		}
+		UnrealUSDWrapper.CheckAndSetupUsdSdk(Target, this);
 
 		if (Target.bBuildEditor)
 		{
@@ -52,33 +44,4 @@ public class USDImporterMDL : ModuleRules
 			}
 		}
 	}
-	bool EnableUsdSdk(ReadOnlyTargetRules Target)
-	{
-		bool bEnableUsdSdk = (
-			Target.WindowsPlatform.Compiler != WindowsCompiler.Intel &&
-			Target.StaticAnalyzer == StaticAnalyzer.None
-		);
-
-		// Don't enable USD when running the include tool because it has issues parsing Boost headers
-		if (Target.GlobalDefinitions.Contains("UE_INCLUDE_TOOL=1"))
-		{
-			bEnableUsdSdk = false;
-		}
-
-		if (!Target.bCompilePython)
-		{
-			bEnableUsdSdk = false;
-		}
-
-		bool bIsGameOrMonolithic = Target.LinkType == TargetLinkType.Monolithic || Target.Type == TargetType.Client || Target.Type == TargetType.Server || Target.Type == TargetType.Game;
-		if (bEnableUsdSdk &&
-			bIsGameOrMonolithic &&
-			((Target.StaticAllocator != StaticAllocatorType.Ansi) || !Target.GlobalDefinitions.Contains("UE_USE_MALLOC_FILL_BYTES=0"))
-		)
-		{
-			bEnableUsdSdk = false;
-		}
-
-		return bEnableUsdSdk;
-	}
 }
diff --git a/Engine/Plugins/Runtime/USDCore/Source/USDUtilities/USDUtilities.Build.cs b/Engine/Plugins/Runtime/USDCore/Source/USDUtilities/USDUtilities.Build.cs
--- a/Engine/Plugins/Runtime/USDCore/Source/USDUtilities/USDUtilities.Build.cs
+++ b/Engine/Plugins/Runtime/USDCore/Source/USDUtilities/USDUtilities.Build.cs
@@ -87,46 +87,7 @@ namespace UnrealBuildTool.Rules
 
 			PublicDefinitions.Add("ENABLE_USD_MATERIALX=" + (bEnableUsdMaterialX ? "1" : "0"));
 
-			bool bUseUsdMemoryOverloadsIfSupported = true;
-			bool bEnableUsdMemoryOverloads = (EnableUsdSdk(Target) && bUseUsdMemoryOverloadsIfSupported) &&
-				((Target.Platform == UnrealTargetPlatform.Win64 && Target.Architecture == UnrealArch.X64) //||
-				//Target.Platform == UnrealTargetPlatform.Mac ||
-				//Target.Platform == UnrealTargetPlatform.Linux
-				);
-			if (!bEnableUsdMemoryOverloads)
-			{
-				PrivateDefinitions.Add("SUPPRESS_PER_MODULE_INLINE_FILE"); // This module does not use core's standard operator new/delete overloads
-			}
-		}
-
-		bool EnableUsdSdk(ReadOnlyTargetRules Target)
-		{
-			bool bEnableUsdSdk = (
-				Target.WindowsPlatform.Compiler != WindowsCompiler.Intel &&
-				Target.StaticAnalyzer == StaticAnalyzer.None
-			);
-
-			// Don't enable USD when running the include tool because it has issues parsing Boost headers
-			if (Target.GlobalDefinitions.Contains("UE_INCLUDE_TOOL=1"))
-			{
-				bEnableUsdSdk = false;
-			}
-
-			if (!Target.bCompilePython)
-			{
-				bEnableUsdSdk = false;
-			}
-
-			bool bIsGameOrMonolithic = Target.LinkType == TargetLinkType.Monolithic || Target.Type == TargetType.Client || Target.Type == TargetType.Server || Target.Type == TargetType.Game;
-			if (bEnableUsdSdk &&
-				bIsGameOrMonolithic &&
-				((Target.StaticAllocator != StaticAllocatorType.Ansi) || !Target.GlobalDefinitions.Contains("UE_USE_MALLOC_FILL_BYTES=0"))
-			)
-			{
-				bEnableUsdSdk = false;
-			}
-
-			return bEnableUsdSdk;
+			UnrealUSDWrapper.CheckAndSetupUsdSdk(Target, this);
 		}
 	}
 }
diff --git a/Engine/Plugins/Runtime/USDCore/Source/UnrealUSDWrapper/UnrealUSDWrapper.Build.cs b/Engine/Plugins/Runtime/USDCore/Source/UnrealUSDWrapper/UnrealUSDWrapper.Build.cs
--- a/Engine/Plugins/Runtime/USDCore/Source/UnrealUSDWrapper/UnrealUSDWrapper.Build.cs
+++ b/Engine/Plugins/Runtime/USDCore/Source/UnrealUSDWrapper/UnrealUSDWrapper.Build.cs
@@ -8,6 +8,28 @@ namespace UnrealBuildTool.Rules
 {
 	public class UnrealUSDWrapper : ModuleRules
 	{
+		// Attention: When you include UnrealUSDWrapper as a dependency and include USDMemory.h, use the USD scoped allocs or IMPLEMENT_MODULE_USD
+		// within your module, make sure that you call CheckAndSetupUsdSdk so that your module is properly setup for USD use.
+		public static bool CheckAndSetupUsdSdk(ReadOnlyTargetRules Target, ModuleRules Module)
+		{
+			bool bUseUsdMemoryOverloadsIfSupported = true;
+			bool bIsUsdSdkEnabled = EnableUsdSdk(Target, Module);
+			bool bEnableUsdMemoryOverloads = (bIsUsdSdkEnabled && bUseUsdMemoryOverloadsIfSupported) &&
+				((Target.Platform == UnrealTargetPlatform.Win64 && Target.Architecture == UnrealArch.X64) //||
+				//Target.Platform == UnrealTargetPlatform.Mac ||
+				//Target.Platform == UnrealTargetPlatform.Linux
+				);
+			Module.PublicDefinitions.Add("ENABLE_USD_MEMORY_OVERLOADS=" + (bEnableUsdMemoryOverloads ? "1" : "0"));
+
+			Module.PublicDefinitions.Add("USD_MERGED_MODULES=" + (Target.bMergeModules ? "1" : "0"));
+			if (!Target.bMergeModules && !bEnableUsdMemoryOverloads)
+			{
+				Module.PrivateDefinitions.Add("SUPPRESS_PER_MODULE_INLINE_FILE"); // This module does not use core's standard operator new/delete overloads
+			}
+
+			return bIsUsdSdkEnabled;
+		}
+
 		public UnrealUSDWrapper(ReadOnlyTargetRules Target) : base(Target)
 		{
 			bUseRTTI = true;
@@ -24,16 +46,10 @@ namespace UnrealBuildTool.Rules
 				}
 			);
 
-			bool bUseUsdMemoryOverloadsIfSupported = true;
-			bool bEnableUsdMemoryOverloads = (EnableUsdSdk(Target) && bUseUsdMemoryOverloadsIfSupported) && 
-				((Target.Platform == UnrealTargetPlatform.Win64 && Target.Architecture == UnrealArch.X64) //||
-				//Target.Platform == UnrealTargetPlatform.Mac ||
-				//Target.Platform == UnrealTargetPlatform.Linux
-				);
-			PublicDefinitions.Add("ENABLE_USD_MEMORY_OVERLOADS=" + (bEnableUsdMemoryOverloads ? "1" : "0"));
+			bool bIsUsdSdkEnabled = CheckAndSetupUsdSdk(Target, this);
 
 			// Temporarily disabled runtime USD support until Mac and Linux dynamic linking issues are resolved
-			if (EnableUsdSdk(Target) && (Target.Type == TargetType.Editor || Target.Platform == UnrealTargetPlatform.Win64))
+			if (bIsUsdSdkEnabled && (Target.Type == TargetType.Editor || Target.Platform == UnrealTargetPlatform.Win64))
 			{
 				PublicDependencyModuleNames.Add("Python3");
 
@@ -171,15 +187,9 @@ namespace UnrealBuildTool.Rules
 				PublicDefinitions.Add("USE_USD_SDK=0");
 				PublicDefinitions.Add("USD_USES_SYSTEM_MALLOC=0");
 			}
-
-			PublicDefinitions.Add("USD_MERGED_MODULES=" + (Target.bMergeModules ? "1" : "0"));
-			if (!Target.bMergeModules && !bEnableUsdMemoryOverloads)
-			{
-				PrivateDefinitions.Add("SUPPRESS_PER_MODULE_INLINE_FILE"); // This module does not use core's standard operator new/delete overloads
-			}
 		}
 
-		bool EnableUsdSdk(ReadOnlyTargetRules Target)
+		static bool EnableUsdSdk(ReadOnlyTargetRules Target, ModuleRules Module)
 		{
 			// USD SDK has been built against Python 3 and won't launch if the editor is using Python 2
 
@@ -230,12 +240,12 @@ namespace UnrealBuildTool.Rules
 				((Target.StaticAllocator != StaticAllocatorType.Ansi) || !Target.GlobalDefinitions.Contains("UE_USE_MALLOC_FILL_BYTES=0"))
 			)
 			{
-				PublicDefinitions.Add("USD_FORCE_DISABLED=1");
+				Module.PublicDefinitions.Add("USD_FORCE_DISABLED=1");
 				bEnableUsdSdk = false;
 			}
 			else
 			{
-				PublicDefinitions.Add("USD_FORCE_DISABLED=0");
+				Module.PublicDefinitions.Add("USD_FORCE_DISABLED=0");
 			}
 
 			return bEnableUsdSdk;
