Flesh : add support for variables
- Deprecated direct access to the Dataflow and terminal properties for c++
- added accessors and made blueprint use them
- added a migration path on PostLoad

#Synced-CL 48301667

diff --git a/Engine/Plugins/Experimental/ChaosFlesh/Source/ChaosFleshEditor/Private/ChaosFlesh/Asset/AssetDefinition_FleshAsset.cpp b/Engine/Plugins/Experimental/ChaosFlesh/Source/ChaosFleshEditor/Private/ChaosFlesh/Asset/AssetDefinition_FleshAsset.cpp
--- a/Engine/Plugins/Experimental/ChaosFlesh/Source/ChaosFleshEditor/Private/ChaosFlesh/Asset/AssetDefinition_FleshAsset.cpp
+++ b/Engine/Plugins/Experimental/ChaosFlesh/Source/ChaosFleshEditor/Private/ChaosFlesh/Asset/AssetDefinition_FleshAsset.cpp
@@ -66,7 +66,7 @@ EAssetCommandResult UAssetDefinition_FleshAsset::OpenAssets(const FAssetOpenArgs
 			{
 				if (UDataflow* const NewDataflowAsset = Cast<UDataflow>(UE::DataflowAssetDefinitionHelpers::NewOrOpenDataflowAsset(FleshAsset)))
 				{
-					FleshAsset->DataflowAsset = NewDataflowAsset;
+					FleshAsset->SetDataflowAsset(NewDataflowAsset);
 				}
 			}
 
@@ -100,5 +100,6 @@ EAssetCommandResult UAssetDefinition_FleshAsset::OpenAssets(const FAssetOpenArgs
 	return EAssetCommandResult::Unhandled;
 }
 
+DATAFLOW_REGISTER_ASSET_MENUS(UFleshAsset);
 
 #undef LOCTEXT_NAMESPACE
diff --git a/Engine/Plugins/Experimental/ChaosFlesh/Source/ChaosFleshEngine/Private/ChaosFlesh/FleshAsset.cpp b/Engine/Plugins/Experimental/ChaosFlesh/Source/ChaosFleshEngine/Private/ChaosFlesh/FleshAsset.cpp
--- a/Engine/Plugins/Experimental/ChaosFlesh/Source/ChaosFleshEngine/Private/ChaosFlesh/FleshAsset.cpp
+++ b/Engine/Plugins/Experimental/ChaosFlesh/Source/ChaosFleshEngine/Private/ChaosFlesh/FleshAsset.cpp
@@ -18,6 +18,12 @@
 
 DEFINE_LOG_CATEGORY_STATIC(LogFleshAssetInternal, Log, All);
 
+namespace UE::FleshAsset::Private
+{
+	static const FName DataflowTerminalNodeName("FleshAssetTerminal");
+}
+
+//////////////////////////////////////////////////////////////////////////////////////////////////////////
 
 FFleshAssetEdit::FFleshAssetEdit(UFleshAsset* InAsset, FPostEditFunctionCallback InCallback)
 	: PostEditCallback(InCallback)
@@ -39,10 +45,16 @@ FFleshCollection* FFleshAssetEdit::GetFleshCollection()
 	return nullptr;
 }
 
+//////////////////////////////////////////////////////////////////////////////////////////////////////////
+
 UFleshAsset::UFleshAsset(const FObjectInitializer& ObjectInitializer)
 	: Super(ObjectInitializer)
 	, FleshCollection(new FFleshCollection())
+	, DataflowInstance(this)
 {
+	// set the default name for the terminal node 
+	DataflowInstance.SetDataflowTerminal(UE::FleshAsset::Private::DataflowTerminalNodeName);
+
 }
 
 void UFleshAsset::SetCollection(FFleshCollection* InCollection)
@@ -72,6 +84,14 @@ const TManagedArray<FVector3f>* UFleshAsset::FindPositions() const
 	return FleshCollection->FindAttributeTyped<FVector3f>("Vertex", FGeometryCollection::VerticesGroup);
 }
 
+void UFleshAsset::GetAssetRegistryTags(FAssetRegistryTagsContext Context) const
+{
+	// Dataflow specific tags
+	UE::Dataflow::InstanceUtils::GetAssetRegistryTags(this, Context);
+
+	Super::GetAssetRegistryTags(Context);
+}
+
 /** Serialize */
 void UFleshAsset::Serialize(FArchive& Ar)
 {
@@ -95,12 +115,42 @@ TObjectPtr<UDataflowBaseContent> UFleshAsset::CreateDataflowContent()
 	return SkeletalContent;
 }
 
+const FDataflowInstance& UFleshAsset::GetDataflowInstance() const
+{
+	return DataflowInstance;
+}
+
+FDataflowInstance& UFleshAsset::GetDataflowInstance()
+{
+	return DataflowInstance;
+}
+
+void UFleshAsset::SetDataflowAsset(UDataflow* InDataflowAsset)
+{
+	DataflowInstance.SetDataflowAsset(InDataflowAsset);
+}
+
+UDataflow* UFleshAsset::GetDataflowAsset() const
+{
+	return DataflowInstance.GetDataflowAsset();
+}
+
+void UFleshAsset::SetDataflowTerminal(const FString& InDataflowTerminal)
+{
+	DataflowInstance.SetDataflowTerminal(FName(InDataflowTerminal));
+}
+
+FString UFleshAsset::GetDataflowTerminal() const
+{
+	return DataflowInstance.GetDataflowTerminal().ToString();
+}
+
 void UFleshAsset::WriteDataflowContent(const TObjectPtr<UDataflowBaseContent>& DataflowContent) const
 {
 	if(const TObjectPtr<UDataflowFleshContent> SkeletalContent = Cast<UDataflowFleshContent>(DataflowContent))
 	{
-		SkeletalContent->SetDataflowAsset(DataflowAsset);
-		SkeletalContent->SetDataflowTerminal(DataflowTerminal);
+		SkeletalContent->SetDataflowAsset(DataflowInstance.GetDataflowAsset());
+		SkeletalContent->SetDataflowTerminal(DataflowInstance.GetDataflowTerminal().ToString());
 		
 		SkeletalContent->SetSkeletalMesh(SkeletalMesh, true);
 
@@ -134,6 +184,30 @@ void UFleshAsset::ReadDataflowContent(const TObjectPtr<UDataflowBaseContent>& Da
 	}
 }
 
+void UFleshAsset::PostLoad()
+{
+	Super::PostLoad();
+
+	// migrate deprecated data if necessary 
+	MigrateDeprecatedDataflowData();
+}
+
+void UFleshAsset::MigrateDeprecatedDataflowData()
+{
+#if WITH_EDITORONLY_DATA
+	PRAGMA_DISABLE_DEPRECATION_WARNINGS
+	if (DataflowAsset != nullptr)
+	{
+		DataflowInstance.SetDataflowAsset(DataflowAsset);
+		DataflowInstance.SetDataflowTerminal(FName(DataflowTerminal));
+
+		DataflowAsset = nullptr;
+		DataflowTerminal.Empty();
+	}
+	PRAGMA_ENABLE_DEPRECATION_WARNINGS
+#endif
+}
+
 #if WITH_EDITOR
 void UFleshAsset::PostEditChangeProperty(FPropertyChangedEvent& PropertyChangedEvent)
 {
diff --git a/Engine/Plugins/Experimental/ChaosFlesh/Source/ChaosFleshEngine/Public/ChaosFlesh/FleshAsset.h b/Engine/Plugins/Experimental/ChaosFlesh/Source/ChaosFleshEngine/Public/ChaosFlesh/FleshAsset.h
--- a/Engine/Plugins/Experimental/ChaosFlesh/Source/ChaosFleshEngine/Public/ChaosFlesh/FleshAsset.h
+++ b/Engine/Plugins/Experimental/ChaosFlesh/Source/ChaosFleshEngine/Public/ChaosFlesh/FleshAsset.h
@@ -9,7 +9,9 @@
 #include "Dataflow/DataflowNodeParameters.h"
 #include "Dataflow/DataflowEngineTypes.h"
 #include "Dataflow/DataflowContent.h"
+#include "Dataflow/DataflowInstance.h"
 #include "Engine/StaticMesh.h"
+#include "UObject/AssetRegistryTagsContext.h"
 #include "UObject/ObjectMacros.h"
 
 #include "FleshAsset.generated.h"
@@ -52,7 +54,7 @@ private:
 *
 */
 UCLASS(BlueprintType, customconstructor)
-class CHAOSFLESHENGINE_API UFleshAsset : public UObject, public IDataflowContentOwner
+class CHAOSFLESHENGINE_API UFleshAsset : public UObject, public IDataflowContentOwner, public IDataflowInstanceInterface
 {
 	GENERATED_UCLASS_BODY()
 	friend class FFleshAssetEdit;
@@ -81,7 +83,6 @@ public:
 
 	UFleshAsset(const FObjectInitializer& ObjectInitializer = FObjectInitializer::Get());
 
-
 #if WITH_EDITOR
 	/** Post edit change property */
 	virtual void PostEditChangeProperty(FPropertyChangedEvent& PropertyChangedEvent) override;
@@ -90,7 +91,6 @@ public:
 	void PropagateTransformUpdateToComponents() const;
 #endif //if WITH_EDITOR
 
-
 	/**Editing the collection should only be through the edit object.*/
 	UE_DEPRECATED(5.7, "Use SetFleshCollection instead")
 	void SetCollection(FFleshCollection* InCollection);
@@ -111,22 +111,55 @@ public:
 	}
 
 	void Serialize(FArchive& Ar);
+	virtual void PostLoad() override;
+
+	virtual void GetAssetRegistryTags(FAssetRegistryTagsContext Context) const override;
 
 	//~ Begin IDataflowContentOwner interface
 	virtual TObjectPtr<UDataflowBaseContent> CreateDataflowContent() override;
 	virtual void WriteDataflowContent(const TObjectPtr<UDataflowBaseContent>& DataflowContent) const override;
 	virtual void ReadDataflowContent(const TObjectPtr<UDataflowBaseContent>& DataflowContent) override;
 	//~ End IDataflowContentOwner interface
+
+	//~ Begin IDataflowInstanceInterface
+	virtual const FDataflowInstance& GetDataflowInstance() const;
+	virtual FDataflowInstance& GetDataflowInstance();
+	//~ End IDataflowInstanceInterface
 	
 	//
 	// Dataflow
 	//
-	UPROPERTY(EditAnywhere, Category = "Dataflow")
+#if WITH_EDITORONLY_DATA
+	// this is deprecated data replaced with Setter and Getter but it is exposed to blueprint 
+	// so we use the Setter/Getter to make sure the blueprint will use them 
+	UE_DEPRECATED(5.8, "Direct set/get is no longer authorized, use SetDataflowAsset instead")
+	UPROPERTY(BlueprintReadWrite, BlueprintSetter = SetDataflowAsset, BlueprintGetter = GetDataflowAsset, Category = "Dataflow")
 	TObjectPtr<UDataflow> DataflowAsset;
 
-	UPROPERTY(EditAnywhere, Category = "Dataflow")
+	UE_DEPRECATED(5.8, "Direct set/get is no longer authorized, use SetDataflowTerminal instead")
+	UPROPERTY(BlueprintReadWrite, BlueprintSetter = SetDataflowTerminal, BlueprintGetter = GetDataflowTerminal, Category = "Dataflow")
 	FString DataflowTerminal = "FleshAssetTerminal";
+#endif
+
+	UFUNCTION(BlueprintCallable, Category = "Dataflow")
+	void SetDataflowAsset(UDataflow* InDataflowAsset);
 
+	UFUNCTION(BlueprintCallable, Category = "Dataflow")
+	UDataflow* GetDataflowAsset() const;
+
+	UFUNCTION(BlueprintCallable, Category = "Dataflow")
+	void SetDataflowTerminal(const FString& InDataflowTerminal);
+
+	UFUNCTION(BlueprintCallable, Category = "Dataflow")
+	FString GetDataflowTerminal() const;
+
+private:
+	void MigrateDeprecatedDataflowData();
+
+	UPROPERTY(EditAnywhere, Category = "Dataflow")
+	FDataflowInstance DataflowInstance;
+
+public:
 	UPROPERTY(EditAnywhere, Category = "Dataflow")
 	TArray<FStringValuePair> Overrides;
 
