ChaosFlesh:
-added CompactSurfaceVertices function to extract & compact flesh surface vertices and triangle mesh (for constructing render mesh for flesh)
-supported building render data for flesh
-supported getting reference skeleton for flesh
-added AssetGuid in flesh for asset identification

#Synced-CL 48489583

diff --git a/Engine/Plugins/Experimental/ChaosFlesh/Source/ChaosFlesh/Private/ChaosFlesh/FleshCollectionUtility.cpp b/Engine/Plugins/Experimental/ChaosFlesh/Source/ChaosFlesh/Private/ChaosFlesh/FleshCollectionUtility.cpp
--- a/Engine/Plugins/Experimental/ChaosFlesh/Source/ChaosFlesh/Private/ChaosFlesh/FleshCollectionUtility.cpp
+++ b/Engine/Plugins/Experimental/ChaosFlesh/Source/ChaosFlesh/Private/ChaosFlesh/FleshCollectionUtility.cpp
@@ -152,4 +152,37 @@ namespace ChaosFlesh
 				UE_LOG(LogChaosFlesh, Display, TEXT("WARNING: Encountered %d non-manifold tetrahedral mesh faces."), nonManifold);
 		}
 	}
+
+	TArray<int32> CompactSurfaceVertices(const TArray<FVector3f>& Vertices, const TArray<FIntVector3>& SurfaceElements, TArray<FVector3f>& OutVertices, TArray<FIntVector3>& OutSurfaceElements)
+	{
+		TArray<int32> OldToNewVertexMap;
+		OldToNewVertexMap.Init(INDEX_NONE, Vertices.Num());
+		TArray<int32> NewToOldVertexMap;
+		OutVertices.Reset();
+		OutSurfaceElements.Reset();
+		for (int32 ElementIdx = 0; ElementIdx < SurfaceElements.Num(); ++ElementIdx)
+		{
+			const FIntVector3& Element = SurfaceElements[ElementIdx];
+			if (!Vertices.IsValidIndex(Element[0]) ||
+				!Vertices.IsValidIndex(Element[1]) ||
+				!Vertices.IsValidIndex(Element[2]))
+			{
+				continue;
+			}
+			for (int32 Corner = 0; Corner < 3; ++Corner)
+			{
+				const int32 VertexIdx = Element[Corner];
+				if (OldToNewVertexMap[VertexIdx] < 0)
+				{
+					OldToNewVertexMap[VertexIdx] = OutVertices.Add(Vertices[VertexIdx]);
+					NewToOldVertexMap.Add(VertexIdx);
+				}
+			}
+			OutSurfaceElements.Add(FIntVector3(
+				OldToNewVertexMap[Element[0]],
+				OldToNewVertexMap[Element[1]],
+				OldToNewVertexMap[Element[2]]));
+		}
+		return NewToOldVertexMap;
+	}
 }
diff --git a/Engine/Plugins/Experimental/ChaosFlesh/Source/ChaosFlesh/Public/ChaosFlesh/FleshCollectionUtility.h b/Engine/Plugins/Experimental/ChaosFlesh/Source/ChaosFlesh/Public/ChaosFlesh/FleshCollectionUtility.h
--- a/Engine/Plugins/Experimental/ChaosFlesh/Source/ChaosFlesh/Public/ChaosFlesh/FleshCollectionUtility.h
+++ b/Engine/Plugins/Experimental/ChaosFlesh/Source/ChaosFlesh/Public/ChaosFlesh/FleshCollectionUtility.h
@@ -11,4 +11,6 @@ namespace ChaosFlesh
 {
 	void GetTetFaces(const FIntVector4& Tet, FIntVector3& Face1, FIntVector3& Face2, FIntVector3& Face3, FIntVector3& Face4, const bool invert = false);
 	void GetSurfaceElements(const TArray<FIntVector4>& Tets, TArray<FIntVector3>& SurfaceElements, const bool KeepInteriorFaces, const bool InvertFaces = false);
+	/** Delete vertices that are not used in SurfaceElements and returns a new-to-old vertex map. */
+	CHAOSFLESH_API TArray<int32> CompactSurfaceVertices(const TArray<FVector3f>& Vertices, const TArray<FIntVector3>& SurfaceElements, TArray<FVector3f>& OutVertices, TArray<FIntVector3>& OutSurfaceElements);
 }
diff --git a/Engine/Plugins/Experimental/ChaosFlesh/Source/ChaosFleshEngine/ChaosFleshEngine.Build.cs b/Engine/Plugins/Experimental/ChaosFlesh/Source/ChaosFleshEngine/ChaosFleshEngine.Build.cs
--- a/Engine/Plugins/Experimental/ChaosFlesh/Source/ChaosFleshEngine/ChaosFleshEngine.Build.cs
+++ b/Engine/Plugins/Experimental/ChaosFlesh/Source/ChaosFleshEngine/ChaosFleshEngine.Build.cs
@@ -37,6 +37,7 @@ namespace UnrealBuildTool.Rules
 			PrivateDependencyModuleNames.AddRange(
 				new string[]
 				{
+					"ClothingSystemRuntimeCommon",
 					"ComputeFramework",
 					"CoreUObject",
 					"Chaos",
diff --git a/Engine/Plugins/Experimental/ChaosFlesh/Source/ChaosFleshEngine/Private/ChaosFlesh/FleshAsset.cpp b/Engine/Plugins/Experimental/ChaosFlesh/Source/ChaosFleshEngine/Private/ChaosFlesh/FleshAsset.cpp
--- a/Engine/Plugins/Experimental/ChaosFlesh/Source/ChaosFleshEngine/Private/ChaosFlesh/FleshAsset.cpp
+++ b/Engine/Plugins/Experimental/ChaosFlesh/Source/ChaosFleshEngine/Private/ChaosFlesh/FleshAsset.cpp
@@ -8,12 +8,23 @@
 #include "Components/SkeletalMeshComponent.h"
 #include "Dataflow/DataflowContent.h"
 #include "Engine/SkeletalMesh.h"
+#include "GeometryCollection/Facades/CollectionMeshFacade.h"
+#include "GeometryCollection/Facades/CollectionVertexBoneWeightsFacade.h"
 #include "GeometryCollection/TransformCollection.h"
+#include "Rendering/SkeletalMeshRenderData.h"
+#include "Rendering/SkeletalMeshModel.h"
+#include "ChaosFlesh/ChaosFleshCollectionFacade.h"
+#include "ChaosFlesh/FleshCollectionUtility.h"
+#include "Materials/Material.h"
+#include "MaterialDomain.h"
+
 #if WITH_EDITOR
+#include "BoneWeights.h"
 #include "ChaosFlesh/FleshComponent.h"
+#include "MeshUtilities.h"
 #include "UObject/UObjectIterator.h"
+#include "Utils/ClothingMeshUtils.h"
 #endif
-
 #include UE_INLINE_GENERATED_CPP_BY_NAME(FleshAsset)
 
 DEFINE_LOG_CATEGORY_STATIC(LogFleshAssetInternal, Log, All);
@@ -21,6 +32,13 @@ DEFINE_LOG_CATEGORY_STATIC(LogFleshAssetInternal, Log, All);
 namespace UE::FleshAsset::Private
 {
 	static const FName DataflowTerminalNodeName("FleshAssetTerminal");
+	static FReferenceSkeleton MakeDefaultRefSkeleton()
+	{
+		FReferenceSkeleton DefaultRefSkeleton;
+		FReferenceSkeletonModifier SkeletonModifier(DefaultRefSkeleton, nullptr);
+		SkeletonModifier.Add(FMeshBoneInfo(FName("Root"), FString("Root"), INDEX_NONE), FTransform::Identity);
+		return DefaultRefSkeleton;
+	}
 }
 
 //////////////////////////////////////////////////////////////////////////////////////////////////////////
@@ -69,6 +87,418 @@ void UFleshAsset::SetFleshCollection(TUniquePtr<FFleshCollection>&& InCollection
 	Modify();
 }
 
+#if WITH_EDITOR
+static void BuildLod(FSkeletalMeshLODModel& LODModel, const UFleshAsset& FleshAsset, int32 LodIndex)
+{
+	using namespace ::Chaos::Softs;
+
+	// Start from an empty LODModel
+	LODModel.Empty();
+
+	// Clear the mesh infos, none are stored on this asset
+	LODModel.MaxImportVertex = 0;
+
+	// Set 1 texture coordinate
+	LODModel.NumTexCoords = 1;
+
+	// Init the size of the vertex buffer
+	LODModel.NumVertices = 0;
+
+	// Offset to remap the LOD materials to the asset materials
+	const int32 MaterialOffset = 0;
+
+	const TSharedPtr<const FFleshCollection> FleshCollection = FleshAsset.GetFleshCollection();
+	// Load the mesh utilities module used to optimized the index buffer
+	IMeshUtilities& MeshUtilities = FModuleManager::Get().LoadModuleChecked<IMeshUtilities>("MeshUtilities");
+
+	// Build the sim mesh descriptor for creation of the sections' mesh to mesh mapping data
+	Chaos::FFleshCollectionFacade FleshFacade(*FleshCollection);
+	const int32 NumLodSimVertices = FleshFacade.NumVertices();
+
+	TArray<FVector3f> SimPositions;
+	for (int32 VertexIndex = 0; VertexIndex < NumLodSimVertices; ++VertexIndex)
+	{
+		SimPositions.Add(FleshFacade.Vertex[VertexIndex]);
+	}
+	TArray<uint32> SimIndices;
+	TArray<FIntVector3> SimElements; // for correct data formatting
+	for (int32 TriIndex = 0; TriIndex < FleshFacade.NumFaces(); ++TriIndex)
+	{
+		SimIndices.Add(uint32(FleshFacade.Indices[TriIndex][0]));
+		SimIndices.Add(uint32(FleshFacade.Indices[TriIndex][1]));
+		SimIndices.Add(uint32(FleshFacade.Indices[TriIndex][2]));
+		SimElements.Add(FleshFacade.Indices[TriIndex]);
+	}
+	const ClothingMeshUtils::ClothMeshDesc SourceMesh(
+		SimPositions,
+		SimIndices);  // Let it calculate the averaged normals as to match the simulation data output
+	const bool bIsValidSourceMesh = (SourceMesh.GetPositions().Num() > 0);
+
+	TArray<FVector3f> RenderPositions;
+	TArray<FIntVector3> RenderElements;
+	TArray<int32> NewToOldVertexMap = ChaosFlesh::CompactSurfaceVertices(SimPositions, SimElements, RenderPositions, RenderElements);
+	const int32 NumRenderVertices = RenderPositions.Num();
+	const int32 NumRenderFaces = RenderElements.Num();
+	const int32 NumRenderIndices = NumRenderFaces * 3;
+
+	LODModel.MeshToImportVertexMap.Reserve(NumRenderVertices);
+
+	// Populate this LOD's sections and the LOD index buffer
+	const int32 NumSections = 1;  // Cloth Render Patterns == Skeletal Mesh Sections
+	LODModel.Sections.SetNum(NumSections);
+
+	TArray<UE::Tasks::FTask> PendingTasks;
+	PendingTasks.Reserve(NumSections);
+	int32 BaseIndex = 0;
+
+	TArray<TSet<FBoneIndexType>> ActiveBoneIndicesArray;
+	// Require at least one active bone indices set to merge all batches set and include the exta bones indices at the end.
+	ActiveBoneIndicesArray.SetNum(FMath::Max(NumSections, 1));
+
+	// Struct to store FSkeletalMeshLODModel data by thread
+	struct FSkeletalMeshLODModelData
+	{
+		TArray<uint32> IndexBuffer;
+		TArray<int32> MeshToImportVertexMap;
+		int32 MaxImportVertex;
+		uint32 NumTexCoords;
+	};
+
+	TArray <FSkeletalMeshLODModelData> LODModelDatas;
+	LODModelDatas.SetNum(NumSections);
+
+	// Mutex to protect the source mesh which has mutable attribute which are modified
+	UE::FMutex SourceMeshMutex;
+	for (int32 SectionIndex = 0; SectionIndex < NumSections; ++SectionIndex)
+	{
+		// Keep track of the active bone indices for this LOD model
+		TSet<FBoneIndexType>& ActiveBoneIndices = ActiveBoneIndicesArray[SectionIndex];
+		ActiveBoneIndices.Reserve(FleshAsset.GetRefSkeleton().GetNum());
+
+		FSkelMeshSection& Section = LODModel.Sections[SectionIndex];
+		Section.OriginalDataSectionIndex = SectionIndex;
+
+		Section.BaseIndex = (uint32)BaseIndex;
+		BaseIndex += NumRenderIndices;
+		Section.BaseVertexIndex = LODModel.NumVertices;
+		LODModel.NumVertices += (uint32)NumRenderVertices;
+
+		FSkeletalMeshLODModelData& LODModelData = LODModelDatas[SectionIndex];
+
+		UE::Tasks::FTask PendingTask = UE::Tasks::Launch(UE_SOURCE_LOCATION,
+			// Instance data per thread, data that are not shared through different threads
+			[SectionIndex, &Section, NumRenderFaces, NumRenderIndices, &LODModelData, &ActiveBoneIndices, &RenderPositions, &NumRenderVertices, 
+			&RenderElements, &NewToOldVertexMap,
+			// Const data that should not be modified or protected data by mutex.
+			MaterialOffset, &FleshAsset, &MeshUtilities, bIsValidSourceMesh, &SourceMesh, &SourceMeshMutex, LodIndex]()
+			{
+				const int32 MaterialIndex = MaterialOffset + SectionIndex;
+
+				// Build the section face data (indices)
+				Section.MaterialIndex = (uint16)MaterialIndex;
+				Section.NumTriangles = (uint32)NumRenderFaces;
+
+				TArray<uint32> Indices;
+				Indices.SetNumUninitialized(NumRenderIndices);
+
+				for (int32 FaceIndex = 0; FaceIndex < NumRenderFaces; ++FaceIndex)
+				{
+					const FIntVector3& RenderIndices = RenderElements[FaceIndex];
+					for (int32 VertexIndex = 0; VertexIndex < 3; ++VertexIndex)
+					{
+						const int32 RenderIndex = RenderIndices[VertexIndex];
+						Indices[FaceIndex * 3 + VertexIndex] = (uint32)RenderIndex;
+					}
+				}
+
+				MeshUtilities.CacheOptimizeIndexBuffer(Indices);
+
+				LODModelData.IndexBuffer.Append(MoveTemp(Indices));
+
+				// Build the section vertex data 
+				Section.SoftVertices.SetNumUninitialized(NumRenderVertices);
+				Section.NumVertices = NumRenderVertices;
+
+				// Map reference skeleton bone index to the index in the section's bone map
+				TMap<FBoneIndexType, FBoneIndexType> ReferenceToSectionBoneMap;
+
+				// Track how many bones we added to the section's bone map so far
+				int CurSectionBoneMapNum = 0;
+
+				GeometryCollection::Facades::FVertexBoneWeightsFacade VertexBoneWeightsFacade(*FleshAsset.GetFleshCollection());
+				const TManagedArray<TArray<int32>>* VertexBoneIndices = VertexBoneWeightsFacade.FindBoneIndices();
+				const TManagedArray<TArray<float>>* VertexBoneWeights = VertexBoneWeightsFacade.FindBoneWeights();
+				GeometryCollection::Facades::FCollectionMeshFacade CollectionMeshFacade(*FleshAsset.GetFleshCollection());
+				
+				LODModelData.NumTexCoords = 1;
+				for (int32 VertexIndex = 0; VertexIndex < NumRenderVertices; ++VertexIndex)
+				{
+					const int32 CollectionVertexIndex = NewToOldVertexMap[VertexIndex];
+					// Save the original indices for the newly added vertices
+					LODModelData.MeshToImportVertexMap.Add(VertexIndex);
+					LODModelData.MaxImportVertex = VertexIndex;
+
+					FSoftSkinVertex& SoftVertex = Section.SoftVertices[VertexIndex];
+
+					SoftVertex.Position = RenderPositions[VertexIndex];
+
+					SoftVertex.TangentX = CollectionMeshFacade.TangentUAttribute.IsValid() && 
+						CollectionMeshFacade.TangentUAttribute.IsValidIndex(CollectionVertexIndex) ?
+							CollectionMeshFacade.TangentUAttribute[CollectionVertexIndex] :
+							FVector3f(1, 0, 0);
+					SoftVertex.TangentY = CollectionMeshFacade.TangentVAttribute.IsValid() &&
+						CollectionMeshFacade.TangentVAttribute.IsValidIndex(CollectionVertexIndex) ?
+							CollectionMeshFacade.TangentVAttribute[CollectionVertexIndex] :
+							FVector3f(0, 1, 0);
+					SoftVertex.TangentZ = CollectionMeshFacade.NormalAttribute.IsValid() &&
+						CollectionMeshFacade.NormalAttribute.IsValidIndex(CollectionVertexIndex) ?
+							CollectionMeshFacade.NormalAttribute[CollectionVertexIndex] :
+							FVector3f(0, 0, 1); 
+
+					constexpr bool bSRGB = false; // Avoid linear to srgb conversion
+					SoftVertex.Color = CollectionMeshFacade.ColorAttribute.IsValid() &&
+						CollectionMeshFacade.ColorAttribute.IsValidIndex(CollectionVertexIndex) ?
+							CollectionMeshFacade.ColorAttribute[CollectionVertexIndex].ToFColor(bSRGB) :
+							FColor::Red;
+
+					const TArray<FVector2f> RenderUVs = TArray<FVector2f>({ FVector2f() });
+					FMemory::Memzero(SoftVertex.UVs, sizeof(SoftVertex.UVs));
+					const int32 NumTexCoords = FMath::Min(RenderUVs.Num(), (int32)MAX_TEXCOORDS);
+					for (int32 TexCoord = 0; TexCoord < NumTexCoords; ++TexCoord)
+					{
+						SoftVertex.UVs[TexCoord] = RenderUVs[TexCoord];
+					}
+					LODModelData.NumTexCoords = FMath::Max(LODModelData.NumTexCoords, (uint32)NumTexCoords);
+
+					bool HasWeights = VertexBoneIndices && VertexBoneWeights && (*VertexBoneIndices)[CollectionVertexIndex].Num() && (*VertexBoneWeights)[CollectionVertexIndex].Num();
+					const TArray<int32> BoneIndices = HasWeights ? (*VertexBoneIndices)[CollectionVertexIndex] : TArray<int32>({ 0 });
+					const TArray<float> BoneWeights = HasWeights ? (*VertexBoneWeights)[CollectionVertexIndex] : TArray<float>({ 1.f });;
+					const int32 NumInfluences = BoneIndices.Num();
+					// Add all of the bones that have non-zero influence to the section's bone map and keep track of the order
+					// that we added the reference bone via CurSectionBoneMapNum
+					for (int32 Influence = 0; Influence < NumInfluences; ++Influence)
+					{
+						const FBoneIndexType InfluenceBone = (FBoneIndexType)BoneIndices[Influence];
+
+						if (ReferenceToSectionBoneMap.Contains(InfluenceBone) == false)
+						{
+							ReferenceToSectionBoneMap.Add(InfluenceBone, CurSectionBoneMapNum);
+							++CurSectionBoneMapNum;
+						}
+					}
+
+					int32 Influence = 0;
+					for (; Influence < NumInfluences; ++Influence)
+					{
+						const FBoneIndexType InfluenceBone = (FBoneIndexType)BoneIndices[Influence];
+						const float InWeight = BoneWeights[Influence];
+						const uint16 InfluenceWeight = static_cast<uint16>(InWeight * static_cast<float>(UE::AnimationCore::MaxRawBoneWeight) + 0.5f);
+
+						// FSoftSkinVertex::InfluenceBones contain indices into the section's bone map and not the reference
+						// skeleton, so we need to remap
+						const FBoneIndexType* const MappedIndexPtr = ReferenceToSectionBoneMap.Find(InfluenceBone);
+
+						// ReferenceToSectionBoneMap should always contain InfluenceBone since it was added above
+						checkSlow(MappedIndexPtr);
+						if (MappedIndexPtr != nullptr)
+						{
+							SoftVertex.InfluenceBones[Influence] = *MappedIndexPtr;
+							SoftVertex.InfluenceWeights[Influence] = InfluenceWeight;
+						}
+					}
+
+					for (; Influence < MAX_TOTAL_INFLUENCES; ++Influence)
+					{
+						SoftVertex.InfluenceBones[Influence] = 0;
+						SoftVertex.InfluenceWeights[Influence] = 0;
+					}
+				}
+
+				// Initialize the section bone map
+				Section.BoneMap.SetNumUninitialized(ReferenceToSectionBoneMap.Num());
+				for (const TPair<FBoneIndexType, FBoneIndexType>& Pair : ReferenceToSectionBoneMap)
+				{
+					Section.BoneMap[Pair.Value] = Pair.Key;
+				}
+
+				ActiveBoneIndices.Append(Section.BoneMap);
+
+				// Update max bone influences
+				Section.CalcMaxBoneInfluences();
+				Section.CalcUse16BitBoneIndex();
+
+				// Setup clothing data
+				if (bIsValidSourceMesh)
+				{
+					Section.ClothMappingDataLODs.SetNum(1);  // TODO: LODBias maps for raytracing
+
+					Section.ClothingData.AssetLodIndex = LodIndex;
+					Section.ClothingData.AssetGuid = FleshAsset.GetAssetGuid();  // There is only one cloth asset,
+					Section.CorrespondClothAssetIndex = 0;       // this one
+
+					// Create mapping from scratch
+					TArray<FVector3f> SectionRenderPositions;
+					TArray<FVector3f> SectionRenderNormals;
+					TArray<FVector3f> SectionRenderTangents;
+					SectionRenderPositions.Reserve(NumRenderVertices);
+					SectionRenderNormals.Reserve(NumRenderVertices);
+					SectionRenderTangents.Reserve(NumRenderVertices);
+					for (const FSoftSkinVertex& SoftVert : Section.SoftVertices)
+					{
+						SectionRenderPositions.Add(SoftVert.Position);
+						SectionRenderNormals.Add(SoftVert.TangentZ);
+						SectionRenderTangents.Add(SoftVert.TangentX);
+					}
+					TArray<uint32> SectionRenderIndices;
+					SectionRenderIndices.Reserve(NumRenderIndices);
+					const TArrayView<uint32> SectionIndexBuffer(LODModelData.IndexBuffer.GetData(), NumRenderIndices);
+					for (const uint32 LodModelVertIndex : SectionIndexBuffer)
+					{
+						SectionRenderIndices.Add(LodModelVertIndex - Section.BaseVertexIndex);
+					}
+
+					const ClothingMeshUtils::ClothMeshDesc TargetMesh(
+						SectionRenderPositions,
+						SectionRenderNormals,
+						SectionRenderTangents,
+						SectionRenderIndices);
+
+					// SourceMesh has mutable AAbbtree which is modified so this function cannot be multi threaded
+					SourceMeshMutex.Lock();
+					PRAGMA_DISABLE_DEPRECATION_WARNINGS
+					ClothingMeshUtils::GenerateMeshToMeshVertData(
+						Section.ClothMappingDataLODs[0],
+						TargetMesh,
+						SourceMesh,
+						/*MaxDistances*/ nullptr,
+						/*bSmoothTransition_DEPRECATED*/ true,
+						/*bUseMultipleInfluences_DEPRECATED*/ false,
+						/*SkinningKernelRadius_DEPRECATED*/ 30.f);
+					PRAGMA_ENABLE_DEPRECATION_WARNINGS
+					SourceMeshMutex.Unlock();
+				}
+
+				// Compute the overlapping vertices map (inspired from MeshUtilities::BuildSkeletalMesh)
+				const TArray<FSoftSkinVertex>& SoftVertices = Section.SoftVertices;
+
+				typedef TPair<float, int32> FIndexAndZ;  // Acceleration structure, list of vertex Z / index pairs
+				TArray<FIndexAndZ> IndexAndZs;
+				IndexAndZs.Reserve(NumRenderVertices);
+				for (int32 VertexIndex = 0; VertexIndex < NumRenderVertices; ++VertexIndex)
+				{
+					const FVector3f& Position = SoftVertices[VertexIndex].Position;
+
+					const float Z = 0.30f * Position.X + 0.33f * Position.Y + 0.37f * Position.Z;
+					IndexAndZs.Emplace(Z, VertexIndex);
+				}
+				IndexAndZs.Sort([](const FIndexAndZ& A, const FIndexAndZ& B) { return A.Key < B.Key; });
+
+				for (int32 Index0 = 0; Index0 < IndexAndZs.Num(); ++Index0)
+				{
+					const float Z0 = IndexAndZs[Index0].Key;
+					const uint32 VertexIndex0 = IndexAndZs[Index0].Value;
+					const FVector3f& Position0 = SoftVertices[VertexIndex0].Position;
+
+					// Only need to search forward, since we add pairs both ways
+					for (int32 Index1 = Index0 + 1; Index1 < IndexAndZs.Num() && FMath::Abs(IndexAndZs[Index1].Key - Z0) <= THRESH_POINTS_ARE_SAME; ++Index1)
+					{
+						const uint32 VertexIndex1 = IndexAndZs[Index1].Value;
+						const FVector3f& Position1 = SoftVertices[VertexIndex1].Position;
+
+						if (PointsEqual(Position0, Position1))
+						{
+							// Add to the overlapping map
+							TArray<int32>& SrcValueArray = Section.OverlappingVertices.FindOrAdd(VertexIndex0);
+							SrcValueArray.Add(VertexIndex1);
+
+							TArray<int32>& IterValueArray = Section.OverlappingVertices.FindOrAdd(VertexIndex1);
+							IterValueArray.Add(VertexIndex0);
+						}
+					}
+				}
+			});
+
+		PendingTasks.Add(PendingTask);
+	}
+	UE::Tasks::Wait(PendingTasks);
+
+	for (int32 SectionIndex = 0; SectionIndex < NumSections; ++SectionIndex)
+	{
+		FSkeletalMeshLODModelData& LODModelData = LODModelDatas[SectionIndex];
+		LODModel.IndexBuffer.Append(MoveTemp(LODModelData.IndexBuffer));
+		LODModel.NumTexCoords = FMath::Max(LODModel.NumTexCoords, LODModelData.NumTexCoords);
+		LODModel.MeshToImportVertexMap.Append(MoveTemp(LODModelData.MeshToImportVertexMap));
+		LODModel.MaxImportVertex = LODModelData.MaxImportVertex;
+
+		FSkelMeshSection& Section = LODModel.Sections[SectionIndex];
+		// Must be in single thread
+		// Copy to user section data, otherwise the section data set above would get lost when the user section gets synced
+		FSkelMeshSourceSectionUserData::GetSourceSectionUserData(LODModel.UserSectionsData, Section);
+
+	}
+
+	for (int32 SectionIndex = 1; SectionIndex < NumSections; ++SectionIndex)
+	{
+		ActiveBoneIndicesArray[0].Append(MoveTemp(ActiveBoneIndicesArray[SectionIndex]));
+	}
+
+	// Update the active bone indices on the LOD model
+	LODModel.ActiveBoneIndices = ActiveBoneIndicesArray[0].Array();
+
+	// Ensure parent exists with incoming active bone indices, and the result should be sorted
+	FleshAsset.GetRefSkeleton().EnsureParentsExistAndSort(LODModel.ActiveBoneIndices);
+	LODModel.ActiveBoneIndices.Shrink();
+
+	// Add the extra simulation bones to the list of required bones if any
+	LODModel.RequiredBones = LODModel.ActiveBoneIndices;
+}
+
+void UFleshAsset::BuildResourceForRendering(FSkeletalMeshRenderData& OutSkeletalMeshRenderData) const
+{
+	// Allocate empty entries for each LOD level in source mesh
+	const int32 NumLods = 1;
+
+	FSkeletalMeshModel SkelMeshModel;
+	SkelMeshModel.LODModels.Reset(NumLods);
+	// Rebuild each LOD models
+	for (int32 LodIndex = 0; LodIndex < NumLods; ++LodIndex)
+	{
+		SkelMeshModel.LODModels.Add(new FSkeletalMeshLODModel());
+
+		BuildLod(SkelMeshModel.LODModels[LodIndex], *this, LodIndex);
+		const FSkeletalMeshLODModel* LODModel = &(SkelMeshModel.LODModels[LodIndex]);
+
+		FSkeletalMeshLODRenderData* LODData = new FSkeletalMeshLODRenderData();
+		OutSkeletalMeshRenderData.LODRenderData.Add(LODData);
+		
+		LODData->BuildFromLODModel(LODModel);
+	}
+}
+#endif  // #if WITH_EDITOR
+
+TArray<FSkeletalMaterial> UFleshAsset::GetMaterials() const
+{
+	return TArray<FSkeletalMaterial>({ UMaterial::GetDefaultMaterial(EMaterialDomain::MD_Surface) });
+}
+
+const FReferenceSkeleton& UFleshAsset::GetRefSkeleton() const
+{
+	if (SkeletalMesh)
+	{
+		return SkeletalMesh->GetRefSkeleton();
+	}
+	else
+	{
+		static FReferenceSkeleton DefaultRefSkeleton = UE::FleshAsset::Private::MakeDefaultRefSkeleton();
+		return DefaultRefSkeleton;
+	}
+}
+
+FGuid UFleshAsset::GetAssetGuid() const
+{
+	return AssetGuid;
+}
+
 void UFleshAsset::PostEditCallback()
 {
 	//UE_LOG(LogFleshAssetInternal, Log, TEXT("UFleshAsset::PostEditCallback()"));
diff --git a/Engine/Plugins/Experimental/ChaosFlesh/Source/ChaosFleshEngine/Public/ChaosFlesh/FleshAsset.h b/Engine/Plugins/Experimental/ChaosFlesh/Source/ChaosFleshEngine/Public/ChaosFlesh/FleshAsset.h
--- a/Engine/Plugins/Experimental/ChaosFlesh/Source/ChaosFleshEngine/Public/ChaosFlesh/FleshAsset.h
+++ b/Engine/Plugins/Experimental/ChaosFlesh/Source/ChaosFleshEngine/Public/ChaosFlesh/FleshAsset.h
@@ -18,8 +18,11 @@
 
 class UFleshAsset;
 class UDataflow;
+struct FReferenceSkeleton;
+struct FSkeletalMaterial;
 class USkeletalMesh;
 class USkeleton;
+class FSkeletalMeshRenderData;
 
 /**
 *	FFleshAssetEdit
@@ -89,6 +92,7 @@ public:
 	virtual void PostEditUndo() override;
 	/** Propogate the fact that transform have been updated to all components */
 	void PropagateTransformUpdateToComponents() const;
+	void BuildResourceForRendering(FSkeletalMeshRenderData& OutSkeletalMeshRenderData) const;
 #endif //if WITH_EDITOR
 
 	/**Editing the collection should only be through the edit object.*/
@@ -102,6 +106,10 @@ public:
 	void SetFleshCollection(TUniquePtr<FFleshCollection>&& InCollection);
 	TSharedPtr<const FFleshCollection> GetFleshCollection() const { return FleshCollection; }
 
+	TArray<FSkeletalMaterial> GetMaterials() const;
+	const FReferenceSkeleton& GetRefSkeleton() const;
+	FGuid GetAssetGuid() const;
+
 	TManagedArray<FVector3f>& GetPositions();
 	const TManagedArray<FVector3f>* FindPositions() const;
 
@@ -228,6 +236,10 @@ public:
 	
 #endif // WITH_EDITORONLY_DATA
 
+private:
+	/** A unique identifier as used by the section rendering code. */
+	UPROPERTY()
+	FGuid AssetGuid = FGuid::NewGuid();
 };
 
 
diff --git a/Engine/Plugins/Experimental/ChaosFlesh/Source/ChaosFleshNodes/Private/Dataflow/ChaosFleshMakeFleshAssetNode.cpp b/Engine/Plugins/Experimental/ChaosFlesh/Source/ChaosFleshNodes/Private/Dataflow/ChaosFleshMakeFleshAssetNode.cpp
new file mode 100644
--- /dev/null
+++ b/Engine/Plugins/Experimental/ChaosFlesh/Source/ChaosFleshNodes/Private/Dataflow/ChaosFleshMakeFleshAssetNode.cpp
@@ -0,0 +1,151 @@
+// Copyright Epic Games, Inc. All Rights Reserved.
+
+#include "Dataflow/ChaosFleshMakeFleshAssetNode.h"
+
+#include "ChaosFlesh/ChaosFleshCollectionFacade.h"
+
+
+#include UE_INLINE_GENERATED_CPP_BY_NAME(ChaosFleshMakeFleshAssetNode)
+
+#define LOCTEXT_NAMESPACE "ChaosFleshMakeFleshAssetNode"
+
+FMakeFleshAssetNode::FMakeFleshAssetNode(const UE::Dataflow::FNodeParameters& InParam, FGuid InGuid)
+	: FDataflowNode(InParam, InGuid)
+{
+	// Start with Lod0
+	for (int32 Index = 0; Index < NumInitialCollectionLods; ++Index)
+	{
+		AddPins();
+	}
+	check(GetNumInputs() == NumRequiredInputs + NumInitialCollectionLods);  // Update NumRequiredInputs if you add more Inputs. This is used by Serialize.
+	RegisterOutputConnection(&Asset);
+}
+
+void FMakeFleshAssetNode::Evaluate(UE::Dataflow::FContext& Context, const FDataflowOutput*) const
+{		
+	TArray<TUniquePtr<FFleshCollection>> InClothCollections = GetCleanedCollectionLodValues(Context);
+	const int32 LODIndex = 0; //TODO: support multiple LODs
+
+	TObjectPtr<UFleshAsset> OutFleshAsset = NewObject<UFleshAsset>();
+	if (InClothCollections.Num() > 0)
+	{
+		OutFleshAsset->SetFleshCollection(MoveTemp(InClothCollections[LODIndex]));
+		SetValue(Context, OutFleshAsset, &Asset);
+		return;
+	}
+	SetValue(Context, OutFleshAsset, &Asset);
+}
+
+TArray<UE::Dataflow::FPin> FMakeFleshAssetNode::AddPins()
+{
+	const int32 Index = CollectionLods.AddDefaulted();
+	const FDataflowInput& Input = RegisterInputArrayConnection(GetConnectionReference(Index));
+	return { { UE::Dataflow::FPin::EDirection::INPUT, Input.GetType(), Input.GetName() } };
+}
+
+TArray<UE::Dataflow::FPin> FMakeFleshAssetNode::GetPinsToRemove() const
+{
+	const int32 Index = CollectionLods.Num() - 1;
+	check(CollectionLods.IsValidIndex(Index));
+	if (const FDataflowInput* const Input = FindInput(GetConnectionReference(Index)))
+	{
+		return { { UE::Dataflow::FPin::EDirection::INPUT, Input->GetType(), Input->GetName() } };
+	}
+	return Super::GetPinsToRemove();
+}
+
+void FMakeFleshAssetNode::OnPinRemoved(const UE::Dataflow::FPin& Pin)
+{
+	const int32 Index = CollectionLods.Num() - 1;
+	check(CollectionLods.IsValidIndex(Index));
+#if DO_CHECK
+	const FDataflowInput* const Input = FindInput(GetConnectionReference(Index));
+	check(Input);
+	check(Input->GetName() == Pin.Name);
+	check(Input->GetType() == Pin.Type);
+#endif
+	CollectionLods.SetNum(Index);
+
+	return Super::OnPinRemoved(Pin);
+}
+
+TArray<TUniquePtr<FFleshCollection>> FMakeFleshAssetNode::GetCleanedCollectionLodValues(UE::Dataflow::FContext& Context) const
+{
+	TArray<TUniquePtr<FFleshCollection>> CollectionLodValues;
+	CollectionLodValues.Reserve(CollectionLods.Num());
+
+	int32 LastValidLodIndex = INDEX_NONE;
+	for (int32 LodIndex = 0; LodIndex < CollectionLods.Num(); ++LodIndex)
+	{
+		TUniquePtr<FFleshCollection> FleshCollection(GetValue<FManagedArrayCollection>(Context, GetConnectionReference(LodIndex)).NewCopy<FFleshCollection>());
+		Chaos::FFleshCollectionFacade FleshFacade(*FleshCollection); 
+		if (FleshFacade.IsValid())
+		{
+			LastValidLodIndex = LodIndex;
+			CollectionLodValues.Emplace(MoveTemp(FleshCollection));
+			return CollectionLodValues;
+		}
+		else if (LastValidLodIndex >= 0)
+		{
+			TUniquePtr<FFleshCollection> FleshCollectionCopy(CollectionLodValues[LastValidLodIndex]->NewCopy<FFleshCollection>());
+			CollectionLodValues.Emplace(MoveTemp(FleshCollectionCopy));
+		}
+		else
+		{
+			break;
+		}	
+	}
+	return CollectionLodValues;
+}
+
+UE::Dataflow::TConnectionReference<FManagedArrayCollection> FMakeFleshAssetNode::GetConnectionReference(int32 Index) const
+{
+	return { &CollectionLods[Index], Index, &CollectionLods };
+}
+
+void FMakeFleshAssetNode::PostSerialize(const FArchive& Ar)
+{
+	// because we add pins we need to make sure we restore them when loading
+	// to make sure they can get properly reconnected
+
+	if (Ar.IsLoading())
+	{
+		if (CollectionLods.Num() < NumInitialCollectionLods)
+		{
+			CollectionLods.SetNum(NumInitialCollectionLods);  // In case the FManagedArrayCollection wasn't serialized with the node (pre the WithSerializer trait)
+		}
+		for (int32 Index = 0; Index < NumInitialCollectionLods; ++Index)
+		{
+			check(FindInput(GetConnectionReference(Index)));
+		}
+
+		for (int32 Index = NumInitialCollectionLods; Index < CollectionLods.Num(); ++Index)
+		{
+			FindOrRegisterInputArrayConnection(GetConnectionReference(Index));
+		}
+		if (Ar.IsTransacting())
+		{
+			const int32 OrigNumRegisteredInputs = GetNumInputs();
+			check(OrigNumRegisteredInputs >= NumRequiredInputs + NumInitialCollectionLods);
+			const int32 OrigNumCollections = CollectionLods.Num();
+			const int32 OrigNumRegisteredCollections = (OrigNumRegisteredInputs - NumRequiredInputs);
+			if (OrigNumRegisteredCollections > OrigNumCollections)
+			{
+				// Inputs have been removed.
+				// Temporarily expand Collections so we can get connection references.
+				CollectionLods.SetNum(OrigNumRegisteredCollections);
+				for (int32 Index = OrigNumCollections; Index < CollectionLods.Num(); ++Index)
+				{
+					UnregisterInputConnection(GetConnectionReference(Index));
+				}
+				CollectionLods.SetNum(OrigNumCollections);
+			}
+		}
+		else
+		{
+			ensureAlways(CollectionLods.Num() + NumRequiredInputs == GetNumInputs());
+		}
+	}
+}
+
+#undef LOCTEXT_NAMESPACE
diff --git a/Engine/Plugins/Experimental/ChaosFlesh/Source/ChaosFleshNodes/Private/Dataflow/ChaosFleshNodesPlugin.cpp b/Engine/Plugins/Experimental/ChaosFlesh/Source/ChaosFleshNodes/Private/Dataflow/ChaosFleshNodesPlugin.cpp
--- a/Engine/Plugins/Experimental/ChaosFlesh/Source/ChaosFleshNodes/Private/Dataflow/ChaosFleshNodesPlugin.cpp
+++ b/Engine/Plugins/Experimental/ChaosFlesh/Source/ChaosFleshNodes/Private/Dataflow/ChaosFleshNodesPlugin.cpp
@@ -21,6 +21,7 @@
 #include "Dataflow/ChaosFleshKinematicConstraintNode.h"
 #include "Dataflow/ChaosFleshKinematicMuscleAttachmentsNode.h"
 #include "Dataflow/ChaosFleshKinematicSkeletonConstraintNode.h"
+#include "Dataflow/ChaosFleshMakeFleshAssetNode.h"
 #include "Dataflow/ChaosFleshRadialTetrahedronNode.h"
 #include "Dataflow/ChaosFleshTriangleMeshSimulationPropertiesNode.h"
 #include "Dataflow/ChaosFleshSkeletalMeshConstraintNode.h"
@@ -67,6 +68,7 @@ void IChaosFleshNodesPlugin::StartupModule()
 	DATAFLOW_NODE_REGISTER_CREATION_FACTORY(FKinematicMuscleAttachmentsDataflowNode);
 	DATAFLOW_NODE_REGISTER_CREATION_FACTORY(FKinematicSkeletonConstraintDataflowNode);
 	DATAFLOW_NODE_REGISTER_CREATION_FACTORY(FKinematicSkeletalMeshInitializationDataflowNode);
+	DATAFLOW_NODE_REGISTER_CREATION_FACTORY(FMakeFleshAssetNode);
 	DATAFLOW_NODE_REGISTER_CREATION_FACTORY(FRadialTetrahedronDataflowNodes);
 	DATAFLOW_NODE_REGISTER_CREATION_FACTORY(FReadSkeletalMeshCurvesDataflowNode);
 	DATAFLOW_NODE_REGISTER_CREATION_FACTORY(FSetCollidableVerticesDataflowNode);
diff --git a/Engine/Plugins/Experimental/ChaosFlesh/Source/ChaosFleshNodes/Public/Dataflow/ChaosFleshMakeFleshAssetNode.h b/Engine/Plugins/Experimental/ChaosFlesh/Source/ChaosFleshNodes/Public/Dataflow/ChaosFleshMakeFleshAssetNode.h
new file mode 100644
--- /dev/null
+++ b/Engine/Plugins/Experimental/ChaosFlesh/Source/ChaosFleshNodes/Public/Dataflow/ChaosFleshMakeFleshAssetNode.h
@@ -0,0 +1,46 @@
+// Copyright Epic Games, Inc. All Rights Reserved.
+
+#pragma once 
+
+#include "GeometryCollection/ManagedArrayCollection.h"
+#include "Dataflow/DataflowNode.h"
+#include "ChaosFlesh/FleshAsset.h"
+
+#include "ChaosFleshMakeFleshAssetNode.generated.h"
+
+/** Converts managed array collection to a specified asset. */
+USTRUCT(Meta = (DataflowFlesh))
+struct FMakeFleshAssetNode : public FDataflowNode
+{
+	GENERATED_USTRUCT_BODY()
+	DATAFLOW_NODE_DEFINE_INTERNAL(FMakeFleshAssetNode, "MakeFleshAsset", "Flesh", "Asset")
+
+public:
+	FMakeFleshAssetNode(const UE::Dataflow::FNodeParameters& InParam, FGuid InGuid = FGuid::NewGuid());
+
+private:
+	//~ Begin FDataflowNode interface
+	virtual void Evaluate(UE::Dataflow::FContext& Context, const FDataflowOutput*) const override;
+	virtual TArray<UE::Dataflow::FPin> AddPins() override;
+	virtual bool CanAddPin() const override { return true; }
+	virtual bool CanRemovePin() const override { return CollectionLods.Num() > NumInitialCollectionLods; }
+	virtual TArray<UE::Dataflow::FPin> GetPinsToRemove() const override;
+	virtual void OnPinRemoved(const UE::Dataflow::FPin& Pin) override;
+	virtual void PostSerialize(const FArchive& Ar) override;
+	//~ End FDataflowNode interface
+
+	TArray<TUniquePtr<FFleshCollection>> GetCleanedCollectionLodValues(UE::Dataflow::FContext& Context) const;
+	UE::Dataflow::TConnectionReference<FManagedArrayCollection> GetConnectionReference(int32 Index) const;
+
+	/** Input collection for this LOD. */
+	UPROPERTY()
+	TArray<FManagedArrayCollection> CollectionLods;
+
+
+	UPROPERTY(meta = (DataflowOutput))
+	TObjectPtr<UFleshAsset> Asset;
+
+	// This is for runtime only--used to determine if only properties need to be updated.
+	static constexpr int32 NumRequiredInputs = 0;
+	static constexpr int32 NumInitialCollectionLods = 1;
+};
\ No newline at end of file
