CIS fix: commented out currently unused loop

#Synced-CL 48489583

diff --git a/Engine/Plugins/Experimental/ChaosFlesh/Source/ChaosFleshEngine/Private/ChaosFlesh/FleshAsset.cpp b/Engine/Plugins/Experimental/ChaosFlesh/Source/ChaosFleshEngine/Private/ChaosFlesh/FleshAsset.cpp
--- a/Engine/Plugins/Experimental/ChaosFlesh/Source/ChaosFleshEngine/Private/ChaosFlesh/FleshAsset.cpp
+++ b/Engine/Plugins/Experimental/ChaosFlesh/Source/ChaosFleshEngine/Private/ChaosFlesh/FleshAsset.cpp
@@ -437,10 +437,11 @@ static void BuildLod(FSkeletalMeshLODModel& LODModel, const UFleshAsset& FleshAs
 
 	}
 
-	for (int32 SectionIndex = 1; SectionIndex < NumSections; ++SectionIndex)
-	{
-		ActiveBoneIndicesArray[0].Append(MoveTemp(ActiveBoneIndicesArray[SectionIndex]));
-	}
+	// TODO: enable when NumSections > 1
+	//for (int32 SectionIndex = 1; SectionIndex < NumSections; ++SectionIndex)
+	//{
+	//	ActiveBoneIndicesArray[0].Append(MoveTemp(ActiveBoneIndicesArray[SectionIndex]));
+	//}
 
 	// Update the active bone indices on the LOD model
 	LODModel.ActiveBoneIndices = ActiveBoneIndicesArray[0].Array();
