Dataflow flesh : fix FGenerateSurfaceBindings node not properly handling one of the output causing an assert

#Synced-CL 50212807

diff --git a/Engine/Plugins/Experimental/ChaosFlesh/Source/ChaosFleshNodes/Private/Dataflow/ChaosFleshGenerateSurfaceBindingsNode.cpp b/Engine/Plugins/Experimental/ChaosFlesh/Source/ChaosFleshNodes/Private/Dataflow/ChaosFleshGenerateSurfaceBindingsNode.cpp
--- a/Engine/Plugins/Experimental/ChaosFlesh/Source/ChaosFleshNodes/Private/Dataflow/ChaosFleshGenerateSurfaceBindingsNode.cpp
+++ b/Engine/Plugins/Experimental/ChaosFlesh/Source/ChaosFleshNodes/Private/Dataflow/ChaosFleshGenerateSurfaceBindingsNode.cpp
@@ -192,7 +192,7 @@ FGenerateSurfaceBindings::Evaluate(UE::Dataflow::FContext& Context, const FDataf
 	auto DoubleVert = [](FVector3f V) { return FVector3d(V.X, V.Y, V.Z); };
 	auto TVec3Vert = [](FVector3f V) { return Chaos::TVec3<Chaos::FRealDouble>(V.X, V.Y, V.Z); };
 	
-	if (Out->IsA(&Collection))
+	if (Out->IsA(&Collection) || Out->IsA(&SKMDynamicMesh))
 	{
 		TUniquePtr<FTetrahedralCollection> InCollection(GetValue(Context, &Collection).NewCopy<FTetrahedralCollection>());
 		const TManagedArray<FIntVector4>* Tetrahedron = 
