Dataflow:
- Added flesh VectorField rendering to the new rendering system

#Synced-CL 50422917

diff --git a/Engine/Plugins/Experimental/ChaosFlesh/Source/ChaosFleshEditor/Private/ChaosFlesh/ChaosFleshEditorPlugin.cpp b/Engine/Plugins/Experimental/ChaosFlesh/Source/ChaosFleshEditor/Private/ChaosFlesh/ChaosFleshEditorPlugin.cpp
--- a/Engine/Plugins/Experimental/ChaosFlesh/Source/ChaosFleshEditor/Private/ChaosFlesh/ChaosFleshEditorPlugin.cpp
+++ b/Engine/Plugins/Experimental/ChaosFlesh/Source/ChaosFleshEditor/Private/ChaosFlesh/ChaosFleshEditorPlugin.cpp
@@ -16,6 +16,7 @@
 #include "Editor/FleshEditorStyle.h"
 
 #include "ChaosFleshRendering/FleshTetrahedronRenderableType.h"
+#include "ChaosFleshRendering/FleshVectorFieldRenderableType.h"
 
 #define LOCTEXT_NAMESPACE "FleshEditor"
 
@@ -85,6 +86,7 @@ void IChaosFleshEditorPlugin::StartupModule()
 
 	// Register render types
 	UE::Flesh::Private::RegisterFleshTetrahedronRenderableTypes();
+	UE::Flesh::Private::RegisterFleshVectorFieldRenderableTypes();
 }
 
 void IChaosFleshEditorPlugin::ShutdownModule()
diff --git a/Engine/Plugins/Experimental/ChaosFlesh/Source/ChaosFleshEditor/Private/ChaosFleshRendering/FleshVectorFieldRenderableType.cpp b/Engine/Plugins/Experimental/ChaosFlesh/Source/ChaosFleshEditor/Private/ChaosFleshRendering/FleshVectorFieldRenderableType.cpp
new file mode 100644
--- /dev/null
+++ b/Engine/Plugins/Experimental/ChaosFlesh/Source/ChaosFleshEditor/Private/ChaosFleshRendering/FleshVectorFieldRenderableType.cpp
@@ -0,0 +1,107 @@
+// Copyright Epic Games, Inc. All Rights Reserved.
+
+#include "FleshVectorFieldRenderableType.h"
+
+#include "Drawing/LineSetComponent.h"
+
+#include "Dataflow/DataflowRenderingViewMode.h"
+
+#include "DataflowRendering/DataflowRenderableComponents.h"
+#include "DataflowRendering/DataflowRenderableType.h"
+#include "DataflowRendering/DataflowRenderableTypeInstance.h"
+#include "DataflowRendering/DataflowRenderableTypeRegistry.h"
+#include "DataflowRendering/DataflowRenderableTypeUtils.h"
+
+#include "GeometryCollection/GeometryCollection.h"
+#include "GeometryCollection/GeometryCollectionAlgo.h"
+#include "GeometryCollection/ManagedArrayCollection.h"
+
+#include "Materials/MaterialInterface.h"
+#include "Materials/Material.h"
+
+#include "UObject/ObjectPtr.h"
+
+#include "Field/FieldSystemTypes.h"
+
+namespace UE::Flesh::Private
+{
+	////////////////////////////////////////////////////////////////////////////////////////////////////////////////////////////
+
+	class FFleshVectorFieldRenderableType : public UE::Dataflow::IRenderableType
+	{
+		UE_DATAFLOW_IRENDERABLE_PRIMARY_TYPE(FFieldCollection, Collection);
+		UE_DATAFLOW_IRENDERABLE_RENDER_GROUP(VectorField);
+		UE_DATAFLOW_IRENDERABLE_VIEW_MODE(UE::Dataflow::FDataflowConstruction3DViewMode);
+		UE_DATAFLOW_IRENDERABLE_SETTINGS(UDataflowFleshVectorFieldRenderSettings);
+
+	public:
+		FFleshVectorFieldRenderableType()
+		{
+			constexpr bool bDepthTested = false;
+			Material = (bDepthTested) 
+				? LoadObject<UMaterial>(nullptr, TEXT("/MeshModelingToolsetExp/Materials/PointSetComponentMaterial"))
+				: LoadObject<UMaterial>(nullptr, TEXT("/MeshModelingToolsetExp/Materials/PointSetOverlaidComponentMaterial"));
+		}
+
+	private:
+		virtual bool CanRender(const UE::Dataflow::FRenderableTypeInstance& Instance) const
+		{
+			const FFieldCollection& Collection = GetCollection(Instance);
+			return Collection.HasGroup("VectorField") &&
+				Collection.HasAttribute("Start", "VectorField") &&
+				Collection.HasAttribute("End", "VectorField") &&
+				Collection.HasAttribute("Color", "VectorField");
+		}
+
+		virtual void GetPrimitiveComponents(const UE::Dataflow::FRenderableTypeInstance& Instance, UE::Dataflow::FRenderableComponents& OutComponents) const override
+		{
+			const UDataflowFleshVectorFieldRenderSettings* Settings = Instance.GetTypedRenderSettings<UDataflowFleshVectorFieldRenderSettings>();
+
+			if (Settings)
+			{
+				const FFieldCollection& Collection = GetCollection(Instance);
+
+				if (Collection.HasGroup("VectorField") &&
+					Collection.HasAttribute("Start", "VectorField") &&
+					Collection.HasAttribute("End", "VectorField") &&
+					Collection.HasAttribute("Color", "VectorField"))
+				{
+					TArray<TPair<FVector3f, FVector3f>> VectorField = Collection.GetVectorField();
+					TArray<FLinearColor> Colors = Collection.GetVectorColor();
+
+					TArray<FRenderableLine> RenderableLines;
+					RenderableLines.SetNumUninitialized(VectorField.Num());
+
+					if (VectorField.Num() == Colors.Num())
+					{
+						for (int32 Idx = 0; Idx < VectorField.Num(); ++Idx)
+						{
+							FColor Color = Colors[Idx].ToFColor(true);
+							Color = Color.WithAlpha(255);
+
+							RenderableLines[Idx] = FRenderableLine(FVector(VectorField[Idx].Key), FVector(VectorField[Idx].Value + (VectorField[Idx].Value - VectorField[Idx].Key) * Settings->LengthScalar), Color, Settings->LineThickness);
+						}
+
+						static const FName LineComponentName = TEXT("VectorField");
+
+						if (ULineSetComponent* LineComponent = OutComponents.AddNewComponent<ULineSetComponent>(LineComponentName))
+						{
+							LineComponent->AddLines(RenderableLines);
+							LineComponent->SetLineMaterial(Material);
+						}
+					}
+				}
+			}
+		}
+
+		UMaterialInterface* Material = nullptr;
+	};
+
+
+	///////////////////////////////////////////////////////////////////////////////////////////////////////////////
+	
+	void RegisterFleshVectorFieldRenderableTypes()
+	{
+		UE_DATAFLOW_REGISTER_RENDERABLE_TYPE(FFleshVectorFieldRenderableType);
+	}
+}
\ No newline at end of file
diff --git a/Engine/Plugins/Experimental/ChaosFlesh/Source/ChaosFleshEditor/Private/ChaosFleshRendering/FleshVectorFieldRenderableType.h b/Engine/Plugins/Experimental/ChaosFlesh/Source/ChaosFleshEditor/Private/ChaosFleshRendering/FleshVectorFieldRenderableType.h
new file mode 100644
--- /dev/null
+++ b/Engine/Plugins/Experimental/ChaosFlesh/Source/ChaosFleshEditor/Private/ChaosFleshRendering/FleshVectorFieldRenderableType.h
@@ -0,0 +1,25 @@
+// Copyright Epic Games, Inc. All Rights Reserved.
+
+#pragma once
+
+#include "DataflowRendering/DataflowRenderableTypeSettings.h"
+#include "Math/Color.h"
+
+#include "FleshVectorFieldRenderableType.generated.h"
+
+UCLASS(MinimalAPI, AutoExpandCategories="")
+class UDataflowFleshVectorFieldRenderSettings : public UDataflowRenderableTypeSettings
+{
+	GENERATED_BODY()
+public:
+	UPROPERTY(EditAnywhere, Category = "VectorField")
+	float LineThickness = 2.0f;
+
+	UPROPERTY(EditAnywhere, Category = "VectorField")
+	float LengthScalar = 1.0f;
+};
+
+namespace UE::Flesh::Private
+{
+	void RegisterFleshVectorFieldRenderableTypes();
+}
diff --git a/Engine/Plugins/Runtime/MeshModelingToolset/Source/ModelingComponents/Private/Drawing/LineSetComponent.cpp b/Engine/Plugins/Runtime/MeshModelingToolset/Source/ModelingComponents/Private/Drawing/LineSetComponent.cpp
--- a/Engine/Plugins/Runtime/MeshModelingToolset/Source/ModelingComponents/Private/Drawing/LineSetComponent.cpp
+++ b/Engine/Plugins/Runtime/MeshModelingToolset/Source/ModelingComponents/Private/Drawing/LineSetComponent.cpp
@@ -274,6 +274,20 @@ int32 ULineSetComponent::AddLines(const TArray<FVector>& InStart, const TArray<F
 	return NumLines;
 }
 
+int32 ULineSetComponent::AddLines(const TArray<FRenderableLine>& InLines)
+{
+	const int NumLines = InLines.Num();
+
+	MarkRenderStateDirty();
+
+	for (int32 Idx = 0; Idx < InLines.Num(); ++Idx)
+	{
+		AddLine(InLines[Idx]);
+	}
+
+	return NumLines;
+}
+
 void ULineSetComponent::InsertLine(const int32 ID, const FRenderableLine& OverlayLine)
 {
 	Lines.Insert(ID, OverlayLine);
diff --git a/Engine/Plugins/Runtime/MeshModelingToolset/Source/ModelingComponents/Public/Drawing/LineSetComponent.h b/Engine/Plugins/Runtime/MeshModelingToolset/Source/ModelingComponents/Public/Drawing/LineSetComponent.h
--- a/Engine/Plugins/Runtime/MeshModelingToolset/Source/ModelingComponents/Public/Drawing/LineSetComponent.h
+++ b/Engine/Plugins/Runtime/MeshModelingToolset/Source/ModelingComponents/Public/Drawing/LineSetComponent.h
@@ -68,6 +68,8 @@ public:
 		return AddLine(FRenderableLine(InStart, InEnd, InColor, InThickness, InDepthBias));
 	}
 
+	UE_API int32 AddLines(const TArray<FRenderableLine>& InLines);
+
 	/** Add lines to be rendered using the component. */
 	UFUNCTION(BlueprintCallable, Category="Line Set", meta = (AutoCreateRefTerm = "InColor"))
 	UE_API int32 AddLines(
