Dataflow : continue adding support for the template picker
- Add Set dataflow right click menu for geometry collection
- Add template folder for flesh

#Synced-CL 52411794

diff --git a/Engine/Plugins/Experimental/ChaosFlesh/Source/ChaosFleshEditor/Private/ChaosFlesh/ChaosFleshEditorPlugin.cpp b/Engine/Plugins/Experimental/ChaosFlesh/Source/ChaosFleshEditor/Private/ChaosFlesh/ChaosFleshEditorPlugin.cpp
--- a/Engine/Plugins/Experimental/ChaosFlesh/Source/ChaosFleshEditor/Private/ChaosFlesh/ChaosFleshEditorPlugin.cpp
+++ b/Engine/Plugins/Experimental/ChaosFlesh/Source/ChaosFleshEditor/Private/ChaosFlesh/ChaosFleshEditorPlugin.cpp
@@ -13,6 +13,7 @@
 #include "ChaosFlesh/FleshActor.h"
 #include "ChaosFlesh/FleshAsset.h"
 #include "ChaosFlesh/FleshComponent.h"
+#include "Dataflow/DataflowTemplateRegistry.h"
 #include "Editor/FleshEditorStyle.h"
 
 #include "ChaosFleshRendering/FleshTetrahedronRenderableType.h"
@@ -89,6 +90,10 @@ void IChaosFleshEditorPlugin::StartupModule()
 	// Register render types
 	UE::Flesh::Private::RegisterFleshTetrahedronRenderableTypes();
 	UE::Flesh::Private::RegisterFleshVectorFieldRenderableTypes();
+
+	FDataflowTemplateRegistry::Get().RegisterTemplateFolder(UFleshAsset::StaticClass(),
+		{ TEXT("/ChaosFlesh/Templates"), FSlateIcon("ChaosFleshEditorStyle", "ClassThumbnail.FleshAsset") });
+	
 }
 
 void IChaosFleshEditorPlugin::ShutdownModule()
diff --git a/Engine/Plugins/Experimental/GeometryCollectionPlugin/Source/GeometryCollectionEditor/Private/GeometryCollection/AssetDefinition_GeometryCollection.cpp b/Engine/Plugins/Experimental/GeometryCollectionPlugin/Source/GeometryCollectionEditor/Private/GeometryCollection/AssetDefinition_GeometryCollection.cpp
--- a/Engine/Plugins/Experimental/GeometryCollectionPlugin/Source/GeometryCollectionEditor/Private/GeometryCollection/AssetDefinition_GeometryCollection.cpp
+++ b/Engine/Plugins/Experimental/GeometryCollectionPlugin/Source/GeometryCollectionEditor/Private/GeometryCollection/AssetDefinition_GeometryCollection.cpp
@@ -220,6 +220,8 @@ namespace MenuExtension_GeometryCollection
 					}
 				}));
 
+			UE::DataflowAssetDefinitionHelpers::RegisterSetDataflowAssetMenu(UGeometryCollection::StaticClass());
+
 			UE::DataflowAssetDefinitionHelpers::RegisterEmbedDataflowAssetMenu(UGeometryCollection::StaticClass());
 
 			UE::DataflowAssetDefinitionHelpers::RegisterReEvaluateDataflowAssetMenu(UGeometryCollection::StaticClass());
