Flesh : use the new dataflow template picker when creating a flesh asset
- remove dead code as well

#Synced-CL 52411794

diff --git a/Engine/Plugins/Experimental/ChaosFlesh/Source/ChaosFleshEditor/Private/ChaosFlesh/Asset/FleshAssetFactory.cpp b/Engine/Plugins/Experimental/ChaosFlesh/Source/ChaosFleshEditor/Private/ChaosFlesh/Asset/FleshAssetFactory.cpp
--- a/Engine/Plugins/Experimental/ChaosFlesh/Source/ChaosFleshEditor/Private/ChaosFlesh/Asset/FleshAssetFactory.cpp
+++ b/Engine/Plugins/Experimental/ChaosFlesh/Source/ChaosFleshEditor/Private/ChaosFlesh/Asset/FleshAssetFactory.cpp
@@ -1,6 +1,7 @@
 // Copyright Epic Games, Inc. All Rights Reserved.
 
 #include "ChaosFlesh/Asset/FleshAssetFactory.h"
+#include "Dataflow/AssetDefinition_DataflowAsset.h"
 
 #include "ChaosFlesh/FleshComponent.h"
 #include "Editor.h"
@@ -23,25 +24,13 @@ UFleshAssetFactory::UFleshAssetFactory(const FObjectInitializer& ObjectInitializ
 
 UObject* UFleshAssetFactory::FactoryCreateNew(UClass* Class, UObject* InParent, FName Name, EObjectFlags Flags, UObject* Context, FFeedbackContext* Warn)
 {
-	FTransform LastTransform = FTransform::Identity;
-	TArray< FleshTuple > FleshList;
-
-	TArray<FAssetData> SelectedAssets;
-	GEditor->GetContentBrowserSelections(SelectedAssets);
-
-	for (const FAssetData& AssetData : SelectedAssets)
+	UFleshAsset* NewFlesh = NewObject<UFleshAsset>(InParent, Class, Name, Flags | RF_Transactional | RF_Public | RF_Standalone);
+	if (NewFlesh)
 	{
-		if (AssetData.GetAsset()->IsA<UFleshAsset>())
-		{
-			UFleshComponent *DummyValue(NULL);
-			FleshList.Add(FleshTuple(static_cast<const UFleshAsset *>(AssetData.GetAsset()), DummyValue, FTransform()));
-		}
-
+		UE::DataflowAssetDefinitionHelpers::SetDataflowFromTemplatePicker(NewFlesh);
+		NewFlesh->Modify();
 	}
 
-	UFleshAsset* NewFlesh = NewObject<UFleshAsset>(InParent, Class, Name, Flags | RF_Transactional | RF_Public | RF_Standalone);
-
-	NewFlesh->Modify();
 	return NewFlesh;
 }
 
diff --git a/Engine/Plugins/Experimental/ChaosFlesh/Source/ChaosFleshEditor/Public/ChaosFlesh/Asset/FleshAssetFactory.h b/Engine/Plugins/Experimental/ChaosFlesh/Source/ChaosFleshEditor/Public/ChaosFlesh/Asset/FleshAssetFactory.h
--- a/Engine/Plugins/Experimental/ChaosFlesh/Source/ChaosFleshEditor/Public/ChaosFlesh/Asset/FleshAssetFactory.h
+++ b/Engine/Plugins/Experimental/ChaosFlesh/Source/ChaosFleshEditor/Public/ChaosFlesh/Asset/FleshAssetFactory.h
@@ -10,12 +10,9 @@
 class UFleshAsset;
 class UFleshComponent;
 
-typedef TTuple<const UFleshAsset *, const UFleshComponent *, FTransform> FleshTuple;
-
 /**
 * Factory for FleshAsset
 */
-
 UCLASS()
 class CHAOSFLESHEDITOR_API UFleshAssetFactory : public UFactory
 {
